Chaos API - add some checks and fix some comments

#rnx

diff --git a/Engine/Plugins/Experimental/ChaosRigidPhysicsAsync/Source/ChaosRigidPhysicsAsync/Private/ChaosRigidPhysicsAsync/RigidBodyContainerAsync.cpp b/Engine/Plugins/Experimental/ChaosRigidPhysicsAsync/Source/ChaosRigidPhysicsAsync/Private/ChaosRigidPhysicsAsync/RigidBodyContainerAsync.cpp
--- a/Engine/Plugins/Experimental/ChaosRigidPhysicsAsync/Source/ChaosRigidPhysicsAsync/Private/ChaosRigidPhysicsAsync/RigidBodyContainerAsync.cpp
+++ b/Engine/Plugins/Experimental/ChaosRigidPhysicsAsync/Source/ChaosRigidPhysicsAsync/Private/ChaosRigidPhysicsAsync/RigidBodyContainerAsync.cpp
@@ -67,6 +67,7 @@ namespace Chaos::Rigids::Async
 
 	UE::Physics::FRigidBodyContainerHandle FRigidBodyContainerAsyncGT::GetContainerHandle() const
 	{
+		check(Scene != nullptr);
 		return UE::Physics::FRigidBodyContainerHandle(Scene->GetSceneHandle(), GetContainerId());
 	}
 
@@ -299,6 +300,7 @@ namespace Chaos::Rigids::Async
 
 	UE::Physics::FRigidBodyContainerHandle FRigidBodyContainerAsyncPT::GetContainerHandle() const
 	{
+		check(Scene != nullptr);
 		return UE::Physics::FRigidBodyContainerHandle(Scene->GetSceneHandle(), GetContainerId());
 	}
 
diff --git a/Engine/Plugins/Experimental/ChaosRigidPhysicsAsync/Source/ChaosRigidPhysicsAsync/Private/ChaosRigidPhysicsAsync/RigidFactoryAsync.cpp b/Engine/Plugins/Experimental/ChaosRigidPhysicsAsync/Source/ChaosRigidPhysicsAsync/Private/ChaosRigidPhysicsAsync/RigidFactoryAsync.cpp
--- a/Engine/Plugins/Experimental/ChaosRigidPhysicsAsync/Source/ChaosRigidPhysicsAsync/Private/ChaosRigidPhysicsAsync/RigidFactoryAsync.cpp
+++ b/Engine/Plugins/Experimental/ChaosRigidPhysicsAsync/Source/ChaosRigidPhysicsAsync/Private/ChaosRigidPhysicsAsync/RigidFactoryAsync.cpp
@@ -77,19 +77,23 @@ namespace Chaos::Rigids::Async::Private
 
 	UE::Physics::FRigidSceneHandle CreateSceneFromSolver(FRigidFactoryAsync& InFactory, FPBDRigidsSolver* InSolver)
 	{
+		check(InSolver != nullptr);
 		return FRigidAccessorAsync::CreateSceneFromSolver(InFactory, InSolver);
 	}
 
 	UE::Physics::FRigidBodyHandle CreateBodyFromProxy(const UE::Physics::FRigidSceneHandle& InSceneHandle, FSingleParticlePhysicsProxy* InProxy)
 	{
 		using namespace UE::Physics;
+		check(InProxy != nullptr);
 
 		if (TRigidScenePtr<FGameContextRW> ScenePtr = InSceneHandle.LockRW())
 		{
 			if (FRigidBodyContainerAsyncGT* BodyContainer = GetAsyncBodyContainer(ScenePtr))
 			{
-				FRigidBodyAsyncGT* Body = BodyContainer->CreateRigidBodyForProxy(InProxy);
-				return FRigidBodyHandle(FRigidBodyContainerHandle(InSceneHandle, BodyContainer->GetContainerId()), Body->GetBodyId());
+				if (FRigidBodyAsyncGT* Body = BodyContainer->CreateRigidBodyForProxy(InProxy))
+				{
+					return FRigidBodyHandle(FRigidBodyContainerHandle(InSceneHandle, BodyContainer->GetContainerId()), Body->GetBodyId());
+				}
 			}
 		}
 		return FRigidBodyHandle();
@@ -101,8 +105,10 @@ namespace Chaos::Rigids::Async::Private
 
 		if (TRigidBodyPtr<FGameContextRO> BodyPtr = InBodyHandle.LockRO())
 		{
-			FRigidBodyAsyncGT* Body = GetAsyncBody(BodyPtr);
-			return Body->GetProxy();
+			if (FRigidBodyAsyncGT* Body = GetAsyncBody(BodyPtr))
+			{
+				return Body->GetProxy();
+			}
 		}
 		return nullptr;
 	}
diff --git a/Engine/Plugins/Experimental/ChaosRigidPhysicsAsync/Tests/Private/SceneLockBenchmarks.cpp b/Engine/Plugins/Experimental/ChaosRigidPhysicsAsync/Tests/Private/SceneLockBenchmarks.cpp
--- a/Engine/Plugins/Experimental/ChaosRigidPhysicsAsync/Tests/Private/SceneLockBenchmarks.cpp
+++ b/Engine/Plugins/Experimental/ChaosRigidPhysicsAsync/Tests/Private/SceneLockBenchmarks.cpp
@@ -13,8 +13,8 @@
 namespace Chaos::LowLevelTest
 {
 
-	// This test is only intended to be run with an attached sampling profiler
-	// It effectively loops forever, so don't enable it for regular testing
+	// This test is only intended to be run with an attached sampling profiler.
+	// It loops for a very long time to faciliate sampling, so don't enable it for regular testing.
 	TEST_CASE("ChaosAPIBenchmark", "[.LoopForever][ChaosRBPerBodyLock]")
 	{
 		using namespace Chaos;
@@ -70,8 +70,8 @@ namespace Chaos::LowLevelTest
 		Factory.DestroyRigidScene(SceneHandle);
 	}
 
-	// This test is only intended to be run with an attached sampling profiler
-	// It effectively loops forever, so don't enable it for regular testing
+	// This test is only intended to be run with an attached sampling profiler.
+	// It loops for a very long time to faciliate sampling, so don't enable it for regular testing.
 	TEST_CASE("ChaosAPIBenchmark ProxyComparison", "[.LoopForever][ChaosPerBodyLock]")
 	{
 		using namespace Chaos;
diff --git a/Engine/Plugins/Experimental/ChaosRigidPhysicsAsync/Tests/Private/SceneTests.cpp b/Engine/Plugins/Experimental/ChaosRigidPhysicsAsync/Tests/Private/SceneTests.cpp
--- a/Engine/Plugins/Experimental/ChaosRigidPhysicsAsync/Tests/Private/SceneTests.cpp
+++ b/Engine/Plugins/Experimental/ChaosRigidPhysicsAsync/Tests/Private/SceneTests.cpp
@@ -71,6 +71,7 @@ namespace Chaos::LowLevelTest
 
 			// Create a rigid body
 			FRigidBodyHandle BodyHandle = SceneHandle.LockRW()->CreateRigidBody(MakeRigidName(TEXT("Body")), ERigidMovementType::Dynamic);
+			SceneHandle.LockRW()->RegisterRigidBody(BodyHandle);
 
 			// An example callback
 			FTestRigidSceneModifier SceneModifier;
@@ -82,19 +83,21 @@ namespace Chaos::LowLevelTest
 					{
 						// TODO: Still calls the GT functions so it asserts in SingleParticlePhysicsProxy
 						PTBodyPtr->InitTransform(FTransform(FVector(0.0, 0.0, 100.0)));
-
-						// Not Allowed from Simulation scene handle (GT-only)
-						//PTBody.SetM(200.0f);	
 					}
 				};
 
+			SceneHandle.LockRW()->RegisterModifier(&SceneModifier);
+
 			// Tick physics for a while
 			for (int TickIndex = 0; TickIndex < 100; ++TickIndex)
 			{
 				SceneHandle.LockRW()->Tick(Dt);
 			}
 
+			SceneHandle.LockRW()->UnregisterModifier(&SceneModifier);
+
 			// Destroy the body
+			SceneHandle.LockRW()->UnregisterRigidBody(BodyHandle);
 			SceneHandle.LockRW()->DestroyRigidBody(BodyHandle);
 
 			// Destroy the scene
diff --git a/Engine/Source/Runtime/Experimental/RigidPhysics/Public/RigidPhysics/RigidLog.h b/Engine/Source/Runtime/Experimental/RigidPhysics/Public/RigidPhysics/RigidLog.h
--- a/Engine/Source/Runtime/Experimental/RigidPhysics/Public/RigidPhysics/RigidLog.h
+++ b/Engine/Source/Runtime/Experimental/RigidPhysics/Public/RigidPhysics/RigidLog.h
@@ -14,7 +14,7 @@ namespace UE::Physics::Private
 	template<typename FuncNameType, typename TypeNameType>
 	UE_INTERNAL inline void LogUnsupportedAPI(const FuncNameType& FuncName, const TypeNameType& TypeName)
 	{
-		// This error means a function is being called in a context where it is not supported.
+		// This warning means a function is being called in a context where it is not supported.
 		// E.g., Calling SetMass from the physics thread.
 		UE_LOGFMT(LogChaosRB, Warning, "{0} not supported on {1}", FuncName, TypeName);
 	}
