Chaos API - add separate Static and Dynamic body pools. Mainly to test API in prep for GCs - not really required with current implementation.

#rnx
#Synced-CL 48825514

diff --git a/Engine/Plugins/Experimental/ChaosRigidPhysicsAsync/Source/ChaosRigidPhysicsAsync/Private/ChaosRigidPhysicsAsync/RigidSceneAsync.cpp b/Engine/Plugins/Experimental/ChaosRigidPhysicsAsync/Source/ChaosRigidPhysicsAsync/Private/ChaosRigidPhysicsAsync/RigidSceneAsync.cpp
--- a/Engine/Plugins/Experimental/ChaosRigidPhysicsAsync/Source/ChaosRigidPhysicsAsync/Private/ChaosRigidPhysicsAsync/RigidSceneAsync.cpp
+++ b/Engine/Plugins/Experimental/ChaosRigidPhysicsAsync/Source/ChaosRigidPhysicsAsync/Private/ChaosRigidPhysicsAsync/RigidSceneAsync.cpp
@@ -17,7 +17,8 @@ namespace Chaos::Rigids::Async
 		, SceneSettings(InSceneSettings)
 		, Solver(nullptr)
 		, SpatialAcceleration(nullptr)
-		, BodyPool()
+		, StaticBodyPool()
+		, DynamicBodyPool()
 		, SceneModifiers()
 		, ScenePT()
 		, bOwnsSolver(false)
@@ -33,7 +34,8 @@ namespace Chaos::Rigids::Async
 		, SceneSettings()
 		, Solver(InSolver)
 		, SpatialAcceleration(nullptr)
-		, BodyPool()
+		, StaticBodyPool()
+		, DynamicBodyPool()
 		, SceneModifiers()
 		, ScenePT()
 		, bOwnsSolver(false)
@@ -52,12 +54,15 @@ namespace Chaos::Rigids::Async
 	{
 		using namespace UE::Physics;
 
-		BodyPool = MakeUnique<FRigidBodyPoolAsyncGT>(this, Name);
+		FString StaticName = FString::Format(TEXT("{0} Static"), { Name.Value() });
 
-		FRigidObjectId BodyPoolId = GetBodyContainerRegistry()->AddObject(BodyPool.Get());
-		BodyPool->SetId(BodyPoolId);
+		StaticBodyPool = MakeUnique<FRigidBodyPoolAsyncGT>(this, FRigidDebugName::Make(TEXT("{0} Statics"), Name.Value()));
+		RegisterBodyContainer(StaticBodyPool.Get());
 
-		ScenePT = MakeUnique<FRigidSceneAsyncPT>(Name, BodyPool->GetHandle());
+		DynamicBodyPool = MakeUnique<FRigidBodyPoolAsyncGT>(this, FRigidDebugName::Make(TEXT("{0} Dynamics"), Name.Value()));
+		RegisterBodyContainer(DynamicBodyPool.Get());
+
+		ScenePT = MakeUnique<FRigidSceneAsyncPT>(Name, StaticBodyPool->GetHandle(), DynamicBodyPool->GetHandle());
 	}
 
 	void FRigidSceneAsyncGT::DestructScene()
@@ -66,10 +71,11 @@ namespace Chaos::Rigids::Async
 
 		ScenePT.Reset();
 
-		GetBodyContainerRegistry()->RemoveObject(BodyPool->GetHandle().GetId());
-		BodyPool->SetId(FRigidObjectId());
+		UnregisterBodyContainer(DynamicBodyPool.Get());
+		DynamicBodyPool.Reset();
 
-		BodyPool.Reset();
+		UnregisterBodyContainer(StaticBodyPool.Get());
+		StaticBodyPool.Reset();
 	}
 
 	void FRigidSceneAsyncGT::Startup()
@@ -115,6 +121,22 @@ namespace Chaos::Rigids::Async
 		}
 	}
 
+	void FRigidSceneAsyncGT::RegisterBodyContainer(UE::Physics::IRigidBodyContainer* InContainer)
+	{
+		using namespace UE::Physics;
+
+		FRigidObjectId ContainerId = GetBodyContainerRegistry()->AddObject(InContainer);
+		InContainer->SetId(ContainerId);
+	}
+
+	void FRigidSceneAsyncGT::UnregisterBodyContainer(UE::Physics::IRigidBodyContainer* InContainer)
+	{
+		using namespace UE::Physics;
+
+		GetBodyContainerRegistry()->RemoveObject(InContainer->GetHandle().GetId());
+		InContainer->SetId(FRigidObjectId());
+	}
+
 	UE::Physics::FRigidSceneHandle FRigidSceneAsyncGT::GetHandle() const
 	{
 		return UE::Physics::FRigidSceneHandle(const_cast<FRigidSceneAsyncGT*>(this));
@@ -201,13 +223,32 @@ namespace Chaos::Rigids::Async
 		using namespace UE::Physics;
 		CheckGameWriteLocked();
 
-		FRigidBodyAsyncGT* BodyGT = BodyPool->CreateBody(InName, InMovementType);
+		FRigidBodyAsyncGT* BodyGT = nullptr;
 
-		Solver->EnqueueCommandImmediate(
-			[this, BodyHandle = BodyGT->GetHandle(), Proxy = BodyGT->GetProxy()]()
+		if (InMovementType == ERigidMovementType::Static)
+		{
+			BodyGT = StaticBodyPool->CreateBody(InName, InMovementType);
+			if (BodyGT != nullptr)
 			{
-				ScenePT->CreateBodyPT(BodyHandle, Proxy);
-			});
+				Solver->EnqueueCommandImmediate(
+					[this, BodyHandle = BodyGT->GetHandle(), Proxy = BodyGT->GetProxy()]()
+					{
+						ScenePT->CreateStaticBodyPT(BodyHandle, Proxy);
+					});
+			}
+		}
+		else
+		{
+			BodyGT = DynamicBodyPool->CreateBody(InName, InMovementType);
+			if (BodyGT != nullptr)
+			{
+				Solver->EnqueueCommandImmediate(
+					[this, BodyHandle = BodyGT->GetHandle(), Proxy = BodyGT->GetProxy()]()
+					{
+						ScenePT->CreateDynamicBodyPT(BodyHandle, Proxy);
+					});
+			}
+		}
 
 		return BodyGT;
 	}
@@ -217,17 +258,18 @@ namespace Chaos::Rigids::Async
 		using namespace UE::Physics;
 		CheckGameWriteLocked();
 
-		// TODO_CHAOSAPI: Handle destruction of bodies from other containers
-		if (FRigidBodyAsyncGT* Body = InBody->AsA<FRigidBodyAsyncGT>())
-		{
-			Solver->EnqueueCommandImmediate(
-				[this, BodyHandle = Body->GetHandle()]()
-				{
-					ScenePT->DestroyBodyPT(BodyHandle);
-				});
+		FRigidObjectId BodyPoolId = InBody->GetHandle().GetContainerHandle().GetId();
+		IRigidBodyContainer* BodyPool = (BodyPoolId == StaticBodyPool->GetHandle().GetId()) ? StaticBodyPool.Get() : DynamicBodyPool.Get();
+		UE_RIGIDPHYSICS_CHECK(BodyPool != nullptr);
+		UE_RIGIDPHYSICS_CHECK(BodyPool->GetHandle().GetId() == BodyPoolId);
 
-			BodyPool->DestroyBody(Body);
-		}
+		Solver->EnqueueCommandImmediate(
+			[this, BodyHandle = InBody->GetHandle()]()
+			{
+				ScenePT->DestroyBodyPT(BodyHandle);
+			});
+
+		BodyPool->DestroyBody(InBody);
 	}
 
 	void FRigidSceneAsyncGT::RegisterModifier(UE::Physics::IRigidSceneModifier* Modifier)
@@ -369,12 +411,16 @@ namespace Chaos::Rigids::Async
 {
 	UE_RIGIDPHYSICS_RIGIDTYPED_IMPL(FRigidSceneAsyncPT, UE::Physics::IRigidScene);
 
-	FRigidSceneAsyncPT::FRigidSceneAsyncPT(const UE::Physics::FRigidDebugName& InName, const UE::Physics::FRigidBodyContainerHandle& InBodyPoolHandle)
+	FRigidSceneAsyncPT::FRigidSceneAsyncPT(
+		const UE::Physics::FRigidDebugName& InName,
+		const UE::Physics::FRigidBodyContainerHandle& InStaticBodyPoolHandle,
+		const UE::Physics::FRigidBodyContainerHandle& InDynamicBodyPoolHandle)
 		: Name(InName)
 		, Solver(nullptr)
 		, SceneModifiers()
 	{
-		BodyPool = MakeUnique<FRigidBodyPoolAsyncPT>(InName, InBodyPoolHandle);
+		StaticBodyPool = MakeUnique<FRigidBodyPoolAsyncPT>(InName, InStaticBodyPoolHandle);
+		DynamicBodyPool = MakeUnique<FRigidBodyPoolAsyncPT>(InName, InDynamicBodyPoolHandle);
 	}
 
 	void FRigidSceneAsyncPT::Lock(UE::Physics::ERigidLockType InLockType) const
@@ -464,25 +510,45 @@ namespace Chaos::Rigids::Async
 		CheckSimContext();
 	}
 
-	void FRigidSceneAsyncPT::CreateBodyPT(const UE::Physics::FRigidBodyHandle& InBodyHandle, FSingleParticlePhysicsProxy* Proxy)
+	void FRigidSceneAsyncPT::CreateStaticBodyPT(const UE::Physics::FRigidBodyHandle& InBodyHandle, FSingleParticlePhysicsProxy* Proxy)
 	{
 		CheckSimContext();
 
-		BodyPool->CreateBodyPT(InBodyHandle, Proxy);
+		StaticBodyPool->CreateBodyPT(InBodyHandle, Proxy);
+	}
+
+	void FRigidSceneAsyncPT::CreateDynamicBodyPT(const UE::Physics::FRigidBodyHandle& InBodyHandle, FSingleParticlePhysicsProxy* Proxy)
+	{
+		CheckSimContext();
+
+		DynamicBodyPool->CreateBodyPT(InBodyHandle, Proxy);
 	}
 
 	void FRigidSceneAsyncPT::DestroyBodyPT(const UE::Physics::FRigidBodyHandle& InBodyHandle)
 	{
-		BodyPool->DestroyBodyPT(InBodyHandle);
+		using UE::Physics::IRigidBodyContainer;
+
+		if (IRigidBodyContainer* BodyContainer = GetBodyContainer(InBodyHandle.GetContainerHandle().GetId()))
+		{
+			if (FRigidBodyPoolAsyncPT* BodyPool = BodyContainer->AsAChecked<FRigidBodyPoolAsyncPT>())
+			{
+				BodyPool->DestroyBodyPT(InBodyHandle);
+			}
+		}
 	}
 
 	UE::Physics::IRigidBodyContainer* FRigidSceneAsyncPT::GetBodyContainer(const UE::Physics::FRigidObjectId& InBodyContainerId) const
 	{
 		CheckSimContext();
 
-		if (InBodyContainerId == BodyPool->GetHandle().GetId())
+		// @TODO_CHAOSAPI: Use a registry here when we need to support more containers
+		if (InBodyContainerId == StaticBodyPool->GetHandle().GetId())
+		{
+			return StaticBodyPool.Get();
+		}
+		else if (InBodyContainerId == DynamicBodyPool->GetHandle().GetId())
 		{
-			return BodyPool.Get();
+			return DynamicBodyPool.Get();
 		}
 
 		return nullptr;
diff --git a/Engine/Plugins/Experimental/ChaosRigidPhysicsAsync/Source/ChaosRigidPhysicsAsync/Public/ChaosRigidPhysicsAsync/RigidBodyContainerAsync.h b/Engine/Plugins/Experimental/ChaosRigidPhysicsAsync/Source/ChaosRigidPhysicsAsync/Public/ChaosRigidPhysicsAsync/RigidBodyContainerAsync.h
--- a/Engine/Plugins/Experimental/ChaosRigidPhysicsAsync/Source/ChaosRigidPhysicsAsync/Public/ChaosRigidPhysicsAsync/RigidBodyContainerAsync.h
+++ b/Engine/Plugins/Experimental/ChaosRigidPhysicsAsync/Source/ChaosRigidPhysicsAsync/Public/ChaosRigidPhysicsAsync/RigidBodyContainerAsync.h
@@ -29,13 +29,13 @@ namespace Chaos::Rigids::Async
 		//
 		UE_INTERNAL virtual UE::Physics::FRigidBodyContainerHandle GetHandle() const override final;
 		UE_INTERNAL virtual int32 GetNumBodies() const override final;
+		UE_INTERNAL virtual void SetId(const UE::Physics::FRigidObjectId& InBodyId) override final;
 		UE_INTERNAL virtual FRigidBodyAsyncGT* CreateBody(const UE::Physics::FRigidDebugName& InName, UE::Physics::ERigidMovementType InMovementType) override final;
 		UE_INTERNAL virtual void DestroyBody(UE::Physics::IRigidBody* InBody) override final;
 		//
 		// End IRigidBodyContainer
 		//
 
-		UE_INTERNAL CHAOSRIGIDPHYSICSASYNC_API void SetId(const UE::Physics::FRigidObjectId& InId);
 		UE_INTERNAL CHAOSRIGIDPHYSICSASYNC_API FRigidSceneAsyncGT* GetScene() const;
 		UE_INTERNAL CHAOSRIGIDPHYSICSASYNC_API void VisitBodies(const TFunctionRef<UE::Physics::ERigidVisitorResponse(UE::Physics::IRigidBody*)>& InVisitor);
 
diff --git a/Engine/Plugins/Experimental/ChaosRigidPhysicsAsync/Source/ChaosRigidPhysicsAsync/Public/ChaosRigidPhysicsAsync/RigidSceneAsync.h b/Engine/Plugins/Experimental/ChaosRigidPhysicsAsync/Source/ChaosRigidPhysicsAsync/Public/ChaosRigidPhysicsAsync/RigidSceneAsync.h
--- a/Engine/Plugins/Experimental/ChaosRigidPhysicsAsync/Source/ChaosRigidPhysicsAsync/Public/ChaosRigidPhysicsAsync/RigidSceneAsync.h
+++ b/Engine/Plugins/Experimental/ChaosRigidPhysicsAsync/Source/ChaosRigidPhysicsAsync/Public/ChaosRigidPhysicsAsync/RigidSceneAsync.h
@@ -72,6 +72,8 @@ namespace Chaos::Rigids::Async
 	public:
 		UE_INTERNAL CHAOSRIGIDPHYSICSASYNC_API void Startup();
 		UE_INTERNAL CHAOSRIGIDPHYSICSASYNC_API void Shutdown();
+		UE_INTERNAL CHAOSRIGIDPHYSICSASYNC_API void RegisterBodyContainer(UE::Physics::IRigidBodyContainer* InContainer);
+		UE_INTERNAL CHAOSRIGIDPHYSICSASYNC_API void UnregisterBodyContainer(UE::Physics::IRigidBodyContainer* InContainer);
 		UE_INTERNAL CHAOSRIGIDPHYSICSASYNC_API void ActivateBody(FRigidBodyAsyncGT* InBody);
 		UE_INTERNAL CHAOSRIGIDPHYSICSASYNC_API void DeactivateBody(FRigidBodyAsyncGT* InBody);
 		UE_INTERNAL CHAOSRIGIDPHYSICSASYNC_API void OnRigidBodyTransformUpdated(FRigidBodyAsyncGT* Body);
@@ -92,7 +94,8 @@ namespace Chaos::Rigids::Async
 		FPBDRigidsSolver* Solver;
 		ISpatialAccelerationCollection<FAccelerationStructureHandle, FReal, 3>* SpatialAcceleration;
 
-		TUniquePtr<FRigidBodyPoolAsyncGT> BodyPool;
+		TUniquePtr<FRigidBodyPoolAsyncGT> StaticBodyPool;
+		TUniquePtr<FRigidBodyPoolAsyncGT> DynamicBodyPool;
 
 		TArray<UE::Physics::IRigidSceneModifier*> SceneModifiers;
 
@@ -108,7 +111,10 @@ namespace Chaos::Rigids::Async
 	public:
 		UE_RIGIDPHYSICS_RIGIDTYPED_DECL(CHAOSRIGIDPHYSICSASYNC_API, FRigidSceneAsyncPT);
 
-		UE_INTERNAL CHAOSRIGIDPHYSICSASYNC_API FRigidSceneAsyncPT(const UE::Physics::FRigidDebugName& InName, const UE::Physics::FRigidBodyContainerHandle& InBodyPoolHandle);
+		UE_INTERNAL CHAOSRIGIDPHYSICSASYNC_API FRigidSceneAsyncPT(
+			const UE::Physics::FRigidDebugName& InName, 
+			const UE::Physics::FRigidBodyContainerHandle& InStaticBodyPoolHandle,
+			const UE::Physics::FRigidBodyContainerHandle& InDynamicBodyPoolHandle);
 
 		//
 		// Begin IRigidLockable
@@ -138,7 +144,10 @@ namespace Chaos::Rigids::Async
 		UE_INTERNAL CHAOSRIGIDPHYSICSASYNC_API UE::Physics::FRigidSceneHandle GetSceneHandle() const;
 		UE_INTERNAL CHAOSRIGIDPHYSICSASYNC_API void Startup(UE::Physics::FRigidSceneHandle InSceneHandle, FPBDRigidsSolver* InSolver);
 		UE_INTERNAL CHAOSRIGIDPHYSICSASYNC_API void Shutdown();
-		UE_INTERNAL CHAOSRIGIDPHYSICSASYNC_API void CreateBodyPT(const UE::Physics::FRigidBodyHandle& InBodyHandle, FSingleParticlePhysicsProxy* Proxy);
+		UE_INTERNAL CHAOSRIGIDPHYSICSASYNC_API void RegisterBodyContainer(UE::Physics::IRigidBodyContainer* InContainer);
+		UE_INTERNAL CHAOSRIGIDPHYSICSASYNC_API void UnregisterBodyContainer(UE::Physics::IRigidBodyContainer* InContainer);
+		UE_INTERNAL CHAOSRIGIDPHYSICSASYNC_API void CreateStaticBodyPT(const UE::Physics::FRigidBodyHandle& InBodyHandle, FSingleParticlePhysicsProxy* Proxy);
+		UE_INTERNAL CHAOSRIGIDPHYSICSASYNC_API void CreateDynamicBodyPT(const UE::Physics::FRigidBodyHandle& InBodyHandle, FSingleParticlePhysicsProxy* Proxy);
 		UE_INTERNAL CHAOSRIGIDPHYSICSASYNC_API void DestroyBodyPT(const UE::Physics::FRigidBodyHandle& InBodyHandle);
 		UE_INTERNAL CHAOSRIGIDPHYSICSASYNC_API void ActivateBody(const UE::Physics::FRigidBodyHandle& InBodyHandle);
 		UE_INTERNAL CHAOSRIGIDPHYSICSASYNC_API void DeactivateBody(const UE::Physics::FRigidBodyHandle& InBodyHandle);
@@ -153,7 +162,8 @@ namespace Chaos::Rigids::Async
 		UE::Physics::FRigidDebugName Name;
 		UE::Physics::FRigidSceneHandle SceneHandle;
 		FPBDRigidsSolver* Solver;
-		TUniquePtr<FRigidBodyPoolAsyncPT> BodyPool;
+		TUniquePtr<FRigidBodyPoolAsyncPT> StaticBodyPool;
+		TUniquePtr<FRigidBodyPoolAsyncPT> DynamicBodyPool;
 		TArray<UE::Physics::IRigidSceneModifier*> SceneModifiers;
 	};
 }
diff --git a/Engine/Source/Runtime/Experimental/RigidPhysics/Public/RigidPhysics/RigidBodyContainer.h b/Engine/Source/Runtime/Experimental/RigidPhysics/Public/RigidPhysics/RigidBodyContainer.h
--- a/Engine/Source/Runtime/Experimental/RigidPhysics/Public/RigidPhysics/RigidBodyContainer.h
+++ b/Engine/Source/Runtime/Experimental/RigidPhysics/Public/RigidPhysics/RigidBodyContainer.h
@@ -35,6 +35,15 @@ namespace UE::Physics
 
 		UE_INTERNAL RIGIDPHYSICS_API virtual IRigidBody* GetBody(const FRigidObjectId& InBodyId) const = 0;
 
+		//
+		// Below here is GT only. See if we can separate the Game and SimCallback APIs
+		//
+
+		UE_INTERNAL RIGIDPHYSICS_API virtual void SetId(const FRigidObjectId& InBodyId)
+		{
+			RIGIDPHYSICS_API_UNSUPPORTED();
+		}
+
 		UE_INTERNAL RIGIDPHYSICS_API virtual IRigidBody* CreateBody(const FRigidDebugName& InName, ERigidMovementType InMovementType)
 		{
 			RIGIDPHYSICS_API_UNSUPPORTED();
