Chaos API - RigidSceneAsync now wraps FChaosScene rather than FPBDRIgidsSolver so that it handles the spatial acceleration correctly.

#rnx
#Synced-CL 48825514

diff --git a/Engine/Plugins/Experimental/ChaosRigidPhysicsAsync/Source/ChaosRigidPhysicsAsync/Private/ChaosRigidPhysicsAsync/RigidBodyAsync.cpp b/Engine/Plugins/Experimental/ChaosRigidPhysicsAsync/Source/ChaosRigidPhysicsAsync/Private/ChaosRigidPhysicsAsync/RigidBodyAsync.cpp
--- a/Engine/Plugins/Experimental/ChaosRigidPhysicsAsync/Source/ChaosRigidPhysicsAsync/Private/ChaosRigidPhysicsAsync/RigidBodyAsync.cpp
+++ b/Engine/Plugins/Experimental/ChaosRigidPhysicsAsync/Source/ChaosRigidPhysicsAsync/Private/ChaosRigidPhysicsAsync/RigidBodyAsync.cpp
@@ -235,11 +235,10 @@ namespace Chaos::Rigids::Async
 		ProxyAPI.SetR(InTransform.GetRotation());
 		ProxyAPI.UpdateShapeBounds();
 
-		// TODO_CHAOSAPI: when we have the Spatial Acceleration working
-		//if (FRigidSceneAsyncGT* Scene = GetScene())
-		//{
-		//	Scene->OnRigidBodyTransformUpdated(this);
-		//}
+		if (FRigidSceneAsyncGT* Scene = GetScene())
+		{
+			Scene->OnRigidBodyTransformUpdated(this);
+		}
 	}
 
 	void FRigidBodyAsyncGT::SetKinematicTarget(const FTransform& InTransform)
diff --git a/Engine/Plugins/Experimental/ChaosRigidPhysicsAsync/Source/ChaosRigidPhysicsAsync/Private/ChaosRigidPhysicsAsync/RigidFactoryAsync.cpp b/Engine/Plugins/Experimental/ChaosRigidPhysicsAsync/Source/ChaosRigidPhysicsAsync/Private/ChaosRigidPhysicsAsync/RigidFactoryAsync.cpp
--- a/Engine/Plugins/Experimental/ChaosRigidPhysicsAsync/Source/ChaosRigidPhysicsAsync/Private/ChaosRigidPhysicsAsync/RigidFactoryAsync.cpp
+++ b/Engine/Plugins/Experimental/ChaosRigidPhysicsAsync/Source/ChaosRigidPhysicsAsync/Private/ChaosRigidPhysicsAsync/RigidFactoryAsync.cpp
@@ -54,18 +54,6 @@ namespace Chaos::Rigids::Async
 		return FRigidSceneHandle();
 	}
 
-	UE::Physics::IRigidScene* FRigidFactoryAsync::CreateSceneFromSolver(FPBDRigidsSolver* InSolver)
-	{
-		using namespace UE::Physics;
-
-		FRigidSceneAsyncGT* Scene = new FRigidSceneAsyncGT(FRigidDebugName::Make(*InSolver->GetDebugName().ToString()), InSolver);
-		Scenes.Add(Scene);
-
-		Scene->Startup();
-
-		return Scene;
-	}
-
 	void FRigidFactoryAsync::DestroyScene(const UE::Physics::FRigidSceneHandle& InSceneHandle)
 	{
 		using namespace UE::Physics;
diff --git a/Engine/Plugins/Experimental/ChaosRigidPhysicsAsync/Source/ChaosRigidPhysicsAsync/Private/ChaosRigidPhysicsAsync/RigidSceneAsync.cpp b/Engine/Plugins/Experimental/ChaosRigidPhysicsAsync/Source/ChaosRigidPhysicsAsync/Private/ChaosRigidPhysicsAsync/RigidSceneAsync.cpp
--- a/Engine/Plugins/Experimental/ChaosRigidPhysicsAsync/Source/ChaosRigidPhysicsAsync/Private/ChaosRigidPhysicsAsync/RigidSceneAsync.cpp
+++ b/Engine/Plugins/Experimental/ChaosRigidPhysicsAsync/Source/ChaosRigidPhysicsAsync/Private/ChaosRigidPhysicsAsync/RigidSceneAsync.cpp
@@ -1,5 +1,6 @@
 // Copyright Epic Games, Inc. All Rights Reserved.
 #include "ChaosRigidPhysicsAsync/RigidSceneAsync.h"
+#include "Chaos/ChaosScene.h"
 #include "RigidPhysics/RigidContext.h"
 #include "RigidPhysics/RigidModifier.h"
 
@@ -15,39 +16,40 @@ namespace Chaos::Rigids::Async
 		: UE::Physics::FRigidScene()
 		, Name(InName)
 		, SceneSettings(InSceneSettings)
-		, Solver(nullptr)
-		, SpatialAcceleration(nullptr)
+		, Scene(nullptr)
+		, Owner()
 		, StaticBodyPool()
 		, DynamicBodyPool()
 		, SceneModifiers()
 		, ScenePT()
-		, bOwnsSolver(false)
+		, bOwnsScene(false)
 	{
 		ConstructScene();
 	}
 
-	FRigidSceneAsyncGT::FRigidSceneAsyncGT(
-		const UE::Physics::FRigidDebugName& InName,
-		FPBDRigidsSolver* InSolver)
-		: UE::Physics::FRigidScene()
-		, Name(InName)
-		, SceneSettings()
-		, Solver(InSolver)
-		, SpatialAcceleration(nullptr)
-		, StaticBodyPool()
-		, DynamicBodyPool()
-		, SceneModifiers()
-		, ScenePT()
-		, bOwnsSolver(false)
+	FRigidSceneAsyncGT::~FRigidSceneAsyncGT()
 	{
-		ConstructScene();
+		DestructScene();
 	}
 
-	FRigidSceneAsyncGT::~FRigidSceneAsyncGT()
+#if WITH_ENGINE
+	void FRigidSceneAsyncGT::AddReferencedObjects(FReferenceCollector& Collector)
 	{
-		check(Solver == nullptr);
+		Collector.AddReferencedObject(Owner);
+		if (Scene != nullptr)
+		{
+			Scene->AddReferencedObjects(Collector);
+		}
+	}
+#endif
 
-		DestructScene();
+	FPBDRigidsSolver* FRigidSceneAsyncGT::GetSolver() const
+	{
+		if (Scene != nullptr)
+		{
+			return Scene->GetSolver();
+		}
+		return nullptr;
 	}
 
 	void FRigidSceneAsyncGT::ConstructScene()
@@ -82,43 +84,31 @@ namespace Chaos::Rigids::Async
 	{
 		using namespace UE::Physics;
 
-		if (Solver == nullptr)
+		if (Scene == nullptr)
 		{
-			// TODO_CHAOSAPI: We should contain an FChaosScene here rather than a raw FPBDRigidsSolver because
-			// FChaosScene holds the spatail acceleration. However FPhysicsScene_Chaos inherits FChaosScene so
-			// we need to break that inheritance first. Or more likely, extract the context of FChaosScene into
-			// a new class which can be used here and there.
-
-			// NOTE: Not using full Async for now - it works but it complicates testing (time shift and interpolation)
-			FChaosSolversModule* Module = FChaosSolversModule::GetModule();
-			Solver = Module->CreateSolver(/*Owner=*/nullptr, SceneSettings.AsyncDt, SceneSettings.ThreadingMode, FName(*Name.Value()));
-			bOwnsSolver = true;
+			Scene = new FChaosScene(Owner, SceneSettings.AsyncDt, FName(*Name.Value()));
+			bOwnsScene = true;
 		}
 
-		ScenePT->Startup(GetHandle(), Solver);
+		ScenePT->Startup(GetHandle(), Scene->GetSolver());
 	}
 
 	void FRigidSceneAsyncGT::Shutdown()
 	{
 		ScenePT->Shutdown();
 
-		if (Solver != nullptr)
+		if (Scene != nullptr)
 		{
-			if (bOwnsSolver)
-			{
-				FChaosSolversModule* Module = FChaosSolversModule::GetModule();
-				Module->DestroySolver(Solver);
-				bOwnsSolver = false;
-			}
-			else
+			Scene->Flush();
+
+			if (bOwnsScene)
 			{
-				// From FPhysicsSolverBase::DestroySolver - we have no easy way to flush commands so this will do for now
-				Solver->SetThreadingMode_External(EThreadingModeTemp::SingleThread);
-				Solver->AdvanceAndDispatch_External(0);	//flush any pending commands are executed (for example unregister object)
+				delete Scene;
 			}
-
-			Solver = nullptr;
 		}
+
+		Scene = nullptr;
+		bOwnsScene = false;
 	}
 
 	void FRigidSceneAsyncGT::RegisterBodyContainer(UE::Physics::IRigidBodyContainer* InContainer)
@@ -148,7 +138,7 @@ namespace Chaos::Rigids::Async
 
 		UE_RIGIDPHYSICS_CHECK(IsInGameThreadContext());
 
-		if (Solver != nullptr)
+		if (FPBDRigidsSolver* Solver = GetSolver())
 		{
 			switch (InLockType)
 			{
@@ -168,15 +158,15 @@ namespace Chaos::Rigids::Async
 
 		UE_RIGIDPHYSICS_CHECK(IsInGameThreadContext());
 
-		if (Solver != nullptr)
+		if (FPBDRigidsSolver* Solver = GetSolver())
 		{
 			switch (InLockType)
 			{
 			case ERigidLockType::ReadOnly:
-				Solver->GetExternalDataLock_External().ReadUnlock();
+				Scene->GetSolver()->GetExternalDataLock_External().ReadUnlock();
 				break;
 			case ERigidLockType::ReadWrite:
-				Solver->GetExternalDataLock_External().WriteUnlock();
+				Scene->GetSolver()->GetExternalDataLock_External().WriteUnlock();
 				break;
 			}
 		}
@@ -184,37 +174,46 @@ namespace Chaos::Rigids::Async
 
 	void FRigidSceneAsyncGT::ActivateBody(FRigidBodyAsyncGT* InBody)
 	{
-		if (!InBody->IsActive())
+		if (FPBDRigidsSolver* Solver = GetSolver())
 		{
-			UE_RIGIDPHYSICS_CHECK(InBody->GetProxy() != nullptr);
+			if (!InBody->IsActive())
+			{
+				UE_RIGIDPHYSICS_CHECK(InBody->GetProxy() != nullptr);
 
-			// NOTE: Solver handles GT and PT Particles but we need to create a wrapper on the PT.
-			Solver->RegisterObject(InBody->GetProxy());
+				// NOTE: ChaosScene handles both GT and PT Particles.
+				TArray<FPhysicsActorHandle> ActorHandles;
+				ActorHandles.Add(InBody->GetProxy());
+				Scene->AddActorsToScene_AssumesLocked(ActorHandles);
 
-			// NOTE: Currently this does nothing because Solver->RegisterObject will activate the particle
-			// via the dirty proxy mechanism.
-			Solver->EnqueueCommandImmediate(
-				[BodyHandle = InBody->GetHandle(), this]()
-				{
-					ScenePT->ActivateBody(BodyHandle);
-				});
+				// NOTE: Currently this does nothing because Solver->RegisterObject which is called by
+				// AddActorsToScene_AssumesLocked will activate the particle via the dirty proxy mechanism.
+				Solver->EnqueueCommandImmediate(
+					[BodyHandle = InBody->GetHandle(), this]()
+					{
+						ScenePT->ActivateBody(BodyHandle);
+					});
+			}
 		}
 	}
 
 	void FRigidSceneAsyncGT::DeactivateBody(FRigidBodyAsyncGT* InBody)
 	{
-		if (InBody->IsActive())
+		if (FPBDRigidsSolver* Solver = GetSolver())
 		{
-			UE_RIGIDPHYSICS_CHECK(InBody->GetProxy() != nullptr);
+			if (InBody->IsActive())
+			{
+				UE_RIGIDPHYSICS_CHECK(InBody->GetProxy() != nullptr);
 
-			Solver->EnqueueCommandImmediate(
-				[BodyHandle = InBody->GetHandle(), this]()
-				{
-					ScenePT->DeactivateBody(BodyHandle);
-				});
+				Solver->EnqueueCommandImmediate(
+					[BodyHandle = InBody->GetHandle(), this]()
+					{
+						ScenePT->DeactivateBody(BodyHandle);
+					});
 
-			// NOTE: Solver handles GT and PT
-			Solver->UnregisterObject(InBody->GetProxy());
+				// No FChaosScene::RemoveActorsFromScene_AssumesLocked, but it should do this...
+				FChaosEngineInterface::RemoveActorFromSolver(InBody->GetProxy(), Solver);
+				Scene->RemoveActorFromAccelerationStructure(InBody->GetProxy());
+			}
 		}
 	}
 
@@ -230,11 +229,14 @@ namespace Chaos::Rigids::Async
 			BodyGT = StaticBodyPool->CreateBody(InName, InMovementType);
 			if (BodyGT != nullptr)
 			{
-				Solver->EnqueueCommandImmediate(
-					[this, BodyHandle = BodyGT->GetHandle(), Proxy = BodyGT->GetProxy()]()
-					{
-						ScenePT->CreateStaticBodyPT(BodyHandle, Proxy);
-					});
+				if (FPBDRigidsSolver* Solver = GetSolver())
+				{
+					Solver->EnqueueCommandImmediate(
+						[this, BodyHandle = BodyGT->GetHandle(), Proxy = BodyGT->GetProxy()]()
+						{
+							ScenePT->CreateStaticBodyPT(BodyHandle, Proxy);
+						});
+				}
 			}
 		}
 		else
@@ -242,11 +244,14 @@ namespace Chaos::Rigids::Async
 			BodyGT = DynamicBodyPool->CreateBody(InName, InMovementType);
 			if (BodyGT != nullptr)
 			{
-				Solver->EnqueueCommandImmediate(
-					[this, BodyHandle = BodyGT->GetHandle(), Proxy = BodyGT->GetProxy()]()
-					{
-						ScenePT->CreateDynamicBodyPT(BodyHandle, Proxy);
-					});
+				if (FPBDRigidsSolver* Solver = GetSolver())
+				{
+					Solver->EnqueueCommandImmediate(
+						[this, BodyHandle = BodyGT->GetHandle(), Proxy = BodyGT->GetProxy()]()
+						{
+							ScenePT->CreateDynamicBodyPT(BodyHandle, Proxy);
+						});
+				}
 			}
 		}
 
@@ -260,16 +265,22 @@ namespace Chaos::Rigids::Async
 
 		FRigidObjectId BodyPoolId = InBody->GetHandle().GetContainerHandle().GetId();
 		IRigidBodyContainer* BodyPool = (BodyPoolId == StaticBodyPool->GetHandle().GetId()) ? StaticBodyPool.Get() : DynamicBodyPool.Get();
-		UE_RIGIDPHYSICS_CHECK(BodyPool != nullptr);
-		UE_RIGIDPHYSICS_CHECK(BodyPool->GetHandle().GetId() == BodyPoolId);
 
-		Solver->EnqueueCommandImmediate(
-			[this, BodyHandle = InBody->GetHandle()]()
+		if (BodyPool != nullptr)
+		{
+			UE_RIGIDPHYSICS_CHECK(BodyPool->GetHandle().GetId() == BodyPoolId);
+
+			if (FPBDRigidsSolver* Solver = GetSolver())
 			{
-				ScenePT->DestroyBodyPT(BodyHandle);
-			});
+				Solver->EnqueueCommandImmediate(
+					[this, BodyHandle = InBody->GetHandle()]()
+					{
+						ScenePT->DestroyBodyPT(BodyHandle);
+					});
+			}
 
-		BodyPool->DestroyBody(InBody);
+			BodyPool->DestroyBody(InBody);
+		}
 	}
 
 	void FRigidSceneAsyncGT::RegisterModifier(UE::Physics::IRigidSceneModifier* Modifier)
@@ -278,11 +289,14 @@ namespace Chaos::Rigids::Async
 
 		SceneModifiers.Add(Modifier);
 
-		Solver->EnqueueCommandImmediate(
-			[Modifier, this]()
-			{
-				ScenePT->RegisterModifier(Modifier);
-			});
+		if (FPBDRigidsSolver* Solver = GetSolver())
+		{
+			Solver->EnqueueCommandImmediate(
+				[Modifier, this]()
+				{
+					ScenePT->RegisterModifier(Modifier);
+				});
+		}
 	}
 
 	void FRigidSceneAsyncGT::UnregisterModifier(UE::Physics::IRigidSceneModifier* Modifier)
@@ -291,11 +305,14 @@ namespace Chaos::Rigids::Async
 
 		SceneModifiers.Remove(Modifier);
 
-		Solver->EnqueueCommandImmediate(
+		if (FPBDRigidsSolver* Solver = GetSolver())
+		{
+			Solver->EnqueueCommandImmediate(
 			[Modifier, this]()
 			{
 				ScenePT->UnregisterModifier(Modifier);
 			});
+		}
 	}
 
 	bool FRigidSceneAsyncGT::Query(const UE::Physics::FRigidQueryDesc& QueryDesc, TArray<UE::Physics::FRigidQueryHit>& OutHits)
@@ -310,16 +327,16 @@ namespace Chaos::Rigids::Async
 	{
 		CheckGameWriteLocked();
 
-		if (Solver != nullptr)
-		{
-			OnPreSimulate(Dt);
+		OnPreSimulate(Dt);
 
+		if (FPBDRigidsSolver* Solver = GetSolver())
+		{
 			Solver->AdvanceAndDispatch_External(Dt);
 			Solver->WaitOnPendingTasks_External();
 			Solver->UpdateGameThreadStructures();
-
-			OnPostSimulate(Dt);
 		}
+
+		OnPostSimulate(Dt);
 	}
 
 	void FRigidSceneAsyncGT::OnRigidBodyTransformUpdated(FRigidBodyAsyncGT* Body)
@@ -330,29 +347,12 @@ namespace Chaos::Rigids::Async
 	void FRigidSceneAsyncGT::UpdateRigidBodyInAccelerationStructure(FRigidBodyAsyncGT* Body)
 	{
 		using namespace Chaos;
-		if (Body == nullptr)
-		{
-			return;
-		}
 
-		if (Body->GetProxy() != nullptr)
+		if ((Body != nullptr) && (Body->GetProxy() != nullptr) && (Scene != nullptr))
 		{
-			if (FGeometryParticle* Particle = Body->GetProxy()->GetParticle_LowLevel())
+			if (Body->IsActive())
 			{
-				if (SpatialAcceleration != nullptr)
-				{
-					const FBounds3d Bounds = Body->GetBounds();
-					const FAABB3 WorldBounds = FAABB3(Bounds.Min, Bounds.Max);
-					const bool bHasBounds = !WorldBounds.IsFull();
-
-					FAccelerationStructureHandle AccelerationHandle(Particle);
-					SpatialAcceleration->UpdateElementIn(AccelerationHandle, WorldBounds, bHasBounds, Body->GetProxyAPI().SpatialIdx());
-				}
-
-				if (Solver != nullptr)
-				{
-					Solver->UpdateParticleInAccelerationStructure_External(Particle, EPendingSpatialDataOperation::Update);
-				}
+				Scene->UpdateActorInAccelerationStructure(Body->GetProxy());
 			}
 		}
 	}
diff --git a/Engine/Plugins/Experimental/ChaosRigidPhysicsAsync/Source/ChaosRigidPhysicsAsync/Public/ChaosRigidPhysicsAsync/RigidFactoryAsync.h b/Engine/Plugins/Experimental/ChaosRigidPhysicsAsync/Source/ChaosRigidPhysicsAsync/Public/ChaosRigidPhysicsAsync/RigidFactoryAsync.h
--- a/Engine/Plugins/Experimental/ChaosRigidPhysicsAsync/Source/ChaosRigidPhysicsAsync/Public/ChaosRigidPhysicsAsync/RigidFactoryAsync.h
+++ b/Engine/Plugins/Experimental/ChaosRigidPhysicsAsync/Source/ChaosRigidPhysicsAsync/Public/ChaosRigidPhysicsAsync/RigidFactoryAsync.h
@@ -22,9 +22,6 @@ namespace Chaos::Rigids::Async
 		UE_INTERNAL UE::Physics::FRigidSceneHandle CreateScene(const UE::Physics::FRigidDebugName& InName, const UE::Physics::IRigidSceneSettings* InSettings);
 		UE_INTERNAL virtual void DestroyScene(const UE::Physics::FRigidSceneHandle& InSceneHandle) override final;
 
-		// TODO_CHAOSAPI: Temp until solver is created through new API
-		UE_INTERNAL CHAOSRIGIDPHYSICSASYNC_API UE::Physics::IRigidScene* CreateSceneFromSolver(Chaos::FPBDRigidsSolver* InSolver);
-
 	private:
 		TArray<UE::Physics::IRigidScene*> Scenes;
 	};
diff --git a/Engine/Plugins/Experimental/ChaosRigidPhysicsAsync/Source/ChaosRigidPhysicsAsync/Public/ChaosRigidPhysicsAsync/RigidSceneAsync.h b/Engine/Plugins/Experimental/ChaosRigidPhysicsAsync/Source/ChaosRigidPhysicsAsync/Public/ChaosRigidPhysicsAsync/RigidSceneAsync.h
--- a/Engine/Plugins/Experimental/ChaosRigidPhysicsAsync/Source/ChaosRigidPhysicsAsync/Public/ChaosRigidPhysicsAsync/RigidSceneAsync.h
+++ b/Engine/Plugins/Experimental/ChaosRigidPhysicsAsync/Source/ChaosRigidPhysicsAsync/Public/ChaosRigidPhysicsAsync/RigidSceneAsync.h
@@ -18,6 +18,13 @@
 #include "PhysicsProxy/SingleParticlePhysicsProxy.h"
 #include "RigidPhysics/RigidScene.h"
 
+class FChaosScene;
+
+namespace Chaos
+{
+	class FPBDRigidsSolver;
+}
+
 namespace Chaos::Rigids::Async
 {
 	//
@@ -31,7 +38,12 @@ namespace Chaos::Rigids::Async
 	//
 	//
 
-	class UE_INTERNAL FRigidSceneAsyncGT : public UE::Physics::FRigidScene, public TSharedFromThis<FRigidSceneAsyncGT, ESPMode::ThreadSafe>
+	class UE_INTERNAL FRigidSceneAsyncGT 
+		: public UE::Physics::FRigidScene
+		, public TSharedFromThis<FRigidSceneAsyncGT, ESPMode::ThreadSafe>
+#if WITH_ENGINE
+		, public FGCObject
+#endif
 	{
 	public:
 		UE_RIGIDPHYSICS_RIGIDTYPED_DECL(CHAOSRIGIDPHYSICSASYNC_API, FRigidSceneAsyncGT);
@@ -39,11 +51,22 @@ namespace Chaos::Rigids::Async
 		UE_INTERNAL CHAOSRIGIDPHYSICSASYNC_API FRigidSceneAsyncGT(
 			const UE::Physics::FRigidDebugName& InName,
 			const FRigidSceneSettingsAsync& InSceneSettings);
-		UE_INTERNAL CHAOSRIGIDPHYSICSASYNC_API FRigidSceneAsyncGT(
-			const UE::Physics::FRigidDebugName& InName,
-			FPBDRigidsSolver* InSolver);
 		virtual ~FRigidSceneAsyncGT();
 
+#if WITH_ENGINE
+		//
+		// Begin FGCObject
+		//
+		virtual void AddReferencedObjects(FReferenceCollector& Collector) override final;
+		virtual FString GetReferencerName() const
+		{
+			return "FRigidSceneAsyncGT";
+		}
+		//
+		// End FGCObject
+		//
+#endif
+
 		//
 		// Begin IRigidLockable
 		//
@@ -86,13 +109,14 @@ namespace Chaos::Rigids::Async
 		void OnPreSimulate(FReal Dt);
 		void OnPostSimulate(FReal Dt);
 		void UpdateRigidBodyInAccelerationStructure(FRigidBodyAsyncGT* Body);
+		FPBDRigidsSolver* GetSolver() const;
 
 		UE::Physics::FRigidDebugName Name;
 
 		FRigidSceneSettingsAsync SceneSettings;
 
-		FPBDRigidsSolver* Solver;
-		ISpatialAccelerationCollection<FAccelerationStructureHandle, FReal, 3>* SpatialAcceleration;
+		FChaosScene* Scene;
+		TObjectPtr<UObject> Owner;
 
 		TUniquePtr<FRigidBodyPoolAsyncGT> StaticBodyPool;
 		TUniquePtr<FRigidBodyPoolAsyncGT> DynamicBodyPool;
@@ -101,8 +125,8 @@ namespace Chaos::Rigids::Async
 
 		TUniquePtr<FRigidSceneAsyncPT> ScenePT;
 
-		// TODO_CHAOSAPI: Temp until solver is created through new API
-		bool bOwnsSolver;
+		// TODO_CHAOSAPI: Temp until Scene is always created through new API
+		bool bOwnsScene;
 	};
 
 
