Chaos API - Add Geometry Collection WIP

This is just enough to test the API classes for Geometry Collection but missing the implementation.

#rnx
#Synced-CL 49229388

diff --git a/Engine/Plugins/Experimental/ChaosRigidPhysicsAsync/Source/ChaosRigidPhysicsAsync/Private/ChaosRigidPhysicsAsync/RigidGeometryCollectionAsync.cpp b/Engine/Plugins/Experimental/ChaosRigidPhysicsAsync/Source/ChaosRigidPhysicsAsync/Private/ChaosRigidPhysicsAsync/RigidGeometryCollectionAsync.cpp
new file mode 100644
--- /dev/null
+++ b/Engine/Plugins/Experimental/ChaosRigidPhysicsAsync/Source/ChaosRigidPhysicsAsync/Private/ChaosRigidPhysicsAsync/RigidGeometryCollectionAsync.cpp
@@ -0,0 +1,243 @@
+// Copyright Epic Games, Inc. All Rights Reserved.
+#include "ChaosRigidPhysicsAsync/RigidGeometryCollectionAsync.h"
+#include "ChaosRigidPhysicsAsync/RigidSceneAsync.h"
+#include "PhysicsProxy/GeometryCollectionPhysicsProxy.h"
+
+#if UE_RIGIDPHYSICS_API_ENABLED
+
+namespace Chaos::Rigids::Async
+{
+	UE_RIGIDPHYSICS_RIGIDTYPED_IMPL(FRigidGeometryCollectionBodyAsyncGT, UE::Physics::IRigidBody);
+
+
+	FRigidGeometryCollectionBodyAsyncGT::FRigidGeometryCollectionBodyAsyncGT(FRigidGeometryCollectionAsyncGT* InGeometryCollection)
+	{
+	}
+
+	FRigidGeometryCollectionBodyAsyncGT::~FRigidGeometryCollectionBodyAsyncGT()
+	{
+	}
+
+	UE::Physics::FRigidBodyHandle FRigidGeometryCollectionBodyAsyncGT::GetHandle() const
+	{
+		using namespace UE::Physics;
+
+		return FRigidBodyHandle();
+	}
+
+	bool FRigidGeometryCollectionBodyAsyncGT::IsActive() const
+	{
+		return false;
+	}
+
+	void FRigidGeometryCollectionBodyAsyncGT::Activate()
+	{
+	}
+
+	void FRigidGeometryCollectionBodyAsyncGT::Deactivate()
+	{
+	}
+
+	bool FRigidGeometryCollectionBodyAsyncGT::IsStatic() const
+	{
+		return false;
+	}
+
+	bool FRigidGeometryCollectionBodyAsyncGT::IsKinematic() const
+	{
+		return false;
+	}
+
+	bool FRigidGeometryCollectionBodyAsyncGT::IsDynamic() const
+	{
+		return false;
+	}
+
+	bool FRigidGeometryCollectionBodyAsyncGT::IsSleeping() const
+	{
+		return false;
+	}
+
+	int32 FRigidGeometryCollectionBodyAsyncGT::GetNumShapes() const
+	{
+		return 0;
+	}
+
+	UE::Physics::FConstRigidShapeInstancePtr FRigidGeometryCollectionBodyAsyncGT::GetShape(int32 InShapeIndex) const
+	{
+		using namespace UE::Physics;
+
+		return FConstRigidShapeInstancePtr();
+	}
+}
+
+
+namespace Chaos::Rigids::Async
+{
+	UE_RIGIDPHYSICS_RIGIDTYPED_IMPL(FRigidGeometryCollectionBodyAsyncPT, UE::Physics::IRigidBody);
+
+
+	FRigidGeometryCollectionBodyAsyncPT::FRigidGeometryCollectionBodyAsyncPT(FRigidGeometryCollectionAsyncPT* InGeometryCollection)
+	{
+	}
+
+	FRigidGeometryCollectionBodyAsyncPT::~FRigidGeometryCollectionBodyAsyncPT()
+	{
+	}
+
+	UE::Physics::FRigidBodyHandle FRigidGeometryCollectionBodyAsyncPT::GetHandle() const
+	{
+		using namespace UE::Physics;
+
+		return FRigidBodyHandle();
+	}
+
+	bool FRigidGeometryCollectionBodyAsyncPT::IsActive() const
+	{
+		return false;
+	}
+
+	void FRigidGeometryCollectionBodyAsyncPT::Activate()
+	{
+	}
+
+	void FRigidGeometryCollectionBodyAsyncPT::Deactivate()
+	{
+	}
+
+	bool FRigidGeometryCollectionBodyAsyncPT::IsStatic() const
+	{
+		return false;
+	}
+
+	bool FRigidGeometryCollectionBodyAsyncPT::IsKinematic() const
+	{
+		return false;
+	}
+
+	bool FRigidGeometryCollectionBodyAsyncPT::IsDynamic() const
+	{
+		return false;
+	}
+
+	bool FRigidGeometryCollectionBodyAsyncPT::IsSleeping() const
+	{
+		return false;
+	}
+
+	int32 FRigidGeometryCollectionBodyAsyncPT::GetNumShapes() const
+	{
+		return 0;
+	}
+
+	UE::Physics::FConstRigidShapeInstancePtr FRigidGeometryCollectionBodyAsyncPT::GetShape(int32 InShapeIndex) const
+	{
+		using namespace UE::Physics;
+
+		return FConstRigidShapeInstancePtr();
+	}
+}
+
+
+namespace Chaos::Rigids::Async
+{
+	UE_RIGIDPHYSICS_RIGIDTYPED_IMPL(FRigidGeometryCollectionAsyncGT, UE::Physics::IRigidGeometryCollection);
+
+
+	FRigidGeometryCollectionAsyncGT::FRigidGeometryCollectionAsyncGT(FRigidSceneAsyncGT* InScene, const UE::Physics::FRigidDebugName& InName, const FSimulationParameters& InParameters)
+		: Name(InName)
+		, BodyContainerId()
+		, Scene(InScene)
+		, Proxy(nullptr)
+	{
+		UE_RIGIDPHYSICS_CHECK(Scene != nullptr);
+
+	}
+
+	FRigidGeometryCollectionAsyncGT::~FRigidGeometryCollectionAsyncGT()
+	{
+	}
+
+	UE::Physics::FRigidBodyContainerHandle FRigidGeometryCollectionAsyncGT::GetHandle() const
+	{
+		using namespace UE::Physics;
+
+		return FRigidBodyContainerHandle(Scene->GetHandle(), BodyContainerId);
+	}
+
+	int32 FRigidGeometryCollectionAsyncGT::GetNumBodies() const
+	{
+		if (Proxy != nullptr)
+		{
+			return Proxy->GetNumTransforms();
+		}
+		return 0;
+	}
+
+	FRigidGeometryCollectionBodyAsyncGT* FRigidGeometryCollectionAsyncGT::GetBodyAt(int32 BodyIndex) const
+	{
+		if (Proxy != nullptr)
+		{
+		}
+		return nullptr;
+	}
+
+	FRigidGeometryCollectionBodyAsyncGT* FRigidGeometryCollectionAsyncGT::GetBody(const UE::Physics::FRigidObjectId& InBodyId) const
+	{
+		return nullptr;
+	}
+
+	void FRigidGeometryCollectionAsyncGT::SetId(const UE::Physics::FRigidObjectId& InBodyContainerId)
+	{
+		BodyContainerId = InBodyContainerId;
+	}
+
+}
+
+namespace Chaos::Rigids::Async
+{
+	UE_RIGIDPHYSICS_RIGIDTYPED_IMPL(FRigidGeometryCollectionAsyncPT, UE::Physics::IRigidGeometryCollection);
+
+	FRigidGeometryCollectionAsyncPT::FRigidGeometryCollectionAsyncPT(FRigidSceneAsyncPT* InScene, const UE::Physics::FRigidObjectId& InId, FGeometryCollectionPhysicsProxy* InProxy)
+	 : BodyContainerId(InId)
+	 , Scene(InScene)
+	 , Proxy(InProxy)
+	{
+
+	}
+
+	FRigidGeometryCollectionAsyncPT::~FRigidGeometryCollectionAsyncPT()
+	{
+
+	}
+
+	UE::Physics::FRigidBodyContainerHandle FRigidGeometryCollectionAsyncPT::GetHandle() const
+	{
+		using namespace UE::Physics;
+
+		return FRigidBodyContainerHandle(Scene->GetHandle(), BodyContainerId);
+	}
+
+	void FRigidGeometryCollectionAsyncPT::SetId(const UE::Physics::FRigidObjectId& InBodyContainerId)
+	{
+		// Should not be called on the PT
+		UE_RIGIDPHYSICS_CHECK(false)
+	}
+
+	int32 FRigidGeometryCollectionAsyncPT::GetNumBodies() const
+	{
+		return Proxy->GetNumTransforms();
+	}
+
+	FRigidGeometryCollectionBodyAsyncGT* FRigidGeometryCollectionAsyncPT::GetBody(const UE::Physics::FRigidObjectId& InBodyId) const
+	{
+		return nullptr;
+	}
+
+	FRigidGeometryCollectionBodyAsyncGT* FRigidGeometryCollectionAsyncPT::GetBodyAt(int32 BodyIndex) const
+	{
+		return nullptr;
+	}
+}
+
+#endif // UE_RIGIDPHYSICS_API_ENABLED
\ No newline at end of file
diff --git a/Engine/Plugins/Experimental/ChaosRigidPhysicsAsync/Source/ChaosRigidPhysicsAsync/Private/ChaosRigidPhysicsAsync/RigidSceneAsync.cpp b/Engine/Plugins/Experimental/ChaosRigidPhysicsAsync/Source/ChaosRigidPhysicsAsync/Private/ChaosRigidPhysicsAsync/RigidSceneAsync.cpp
--- a/Engine/Plugins/Experimental/ChaosRigidPhysicsAsync/Source/ChaosRigidPhysicsAsync/Private/ChaosRigidPhysicsAsync/RigidSceneAsync.cpp
+++ b/Engine/Plugins/Experimental/ChaosRigidPhysicsAsync/Source/ChaosRigidPhysicsAsync/Private/ChaosRigidPhysicsAsync/RigidSceneAsync.cpp
@@ -1,5 +1,6 @@
 // Copyright Epic Games, Inc. All Rights Reserved.
 #include "ChaosRigidPhysicsAsync/RigidSceneAsync.h"
+#include "ChaosRigidPhysicsAsync/RigidGeometryCollectionAsync.h"
 #include "Chaos/ChaosScene.h"
 #include "RigidPhysics/RigidContext.h"
 #include "RigidPhysics/RigidModifier.h"
@@ -32,17 +33,6 @@ namespace Chaos::Rigids::Async
 		DestructScene();
 	}
 
-#if WITH_ENGINE
-	void FRigidSceneAsyncGT::AddReferencedObjects(FReferenceCollector& Collector)
-	{
-		Collector.AddReferencedObject(Owner);
-		if (Scene != nullptr)
-		{
-			Scene->AddReferencedObjects(Collector);
-		}
-	}
-#endif
-
 	FPBDRigidsSolver* FRigidSceneAsyncGT::GetSolver() const
 	{
 		if (Scene != nullptr)
@@ -86,7 +76,7 @@ namespace Chaos::Rigids::Async
 
 		if (Scene == nullptr)
 		{
-			Scene = new FChaosScene(Owner, SceneSettings.AsyncDt, FName(*Name.Value()));
+			Scene = new FChaosScene(Owner.Get(), SceneSettings.AsyncDt, FName(*Name.Value()));
 			bOwnsScene = true;
 		}
 
@@ -163,10 +153,10 @@ namespace Chaos::Rigids::Async
 			switch (InLockType)
 			{
 			case ERigidLockType::ReadOnly:
-				Scene->GetSolver()->GetExternalDataLock_External().ReadUnlock();
+				Solver->GetExternalDataLock_External().ReadUnlock();
 				break;
 			case ERigidLockType::ReadWrite:
-				Scene->GetSolver()->GetExternalDataLock_External().WriteUnlock();
+				Solver->GetExternalDataLock_External().WriteUnlock();
 				break;
 			}
 		}
@@ -283,6 +273,45 @@ namespace Chaos::Rigids::Async
 		}
 	}
 
+	FRigidGeometryCollectionAsyncGT* FRigidSceneAsyncGT::CreateGeometryCollection(const UE::Physics::FRigidDebugName& InName, const FSimulationParameters& InParameters)
+	{
+		using namespace UE::Physics;
+
+		FRigidGeometryCollectionAsyncGT* GC = new FRigidGeometryCollectionAsyncGT(this, InName, InParameters);
+
+		RegisterBodyContainer(GC);
+
+		if (FPBDRigidsSolver* Solver = GetSolver())
+		{
+			Solver->EnqueueCommandImmediate(
+				[this, Handle = GC->GetHandle(), Proxy = GC->GetProxy()]()
+				{
+					ScenePT->CreateGeometryCollectionPT(Handle, Proxy);
+				});
+		}
+
+		return GC;
+	}
+
+	void FRigidSceneAsyncGT::DestroyGeometryCollection(UE::Physics::IRigidGeometryCollection* InGC)
+	{
+		if (FRigidGeometryCollectionAsyncGT* GC = InGC->AsAChecked<FRigidGeometryCollectionAsyncGT>())
+		{
+			if (FPBDRigidsSolver* Solver = GetSolver())
+			{
+				Solver->EnqueueCommandImmediate(
+					[this, Handle = GC->GetHandle()]()
+					{
+						ScenePT->DestroyGeometryCollectionPT(Handle);
+					});
+			}
+
+			UnregisterBodyContainer(GC);
+
+			delete GC;
+		}
+	}
+
 	void FRigidSceneAsyncGT::RegisterModifier(UE::Physics::IRigidSceneModifier* Modifier)
 	{
 		CheckGameWriteLocked();
@@ -418,9 +447,24 @@ namespace Chaos::Rigids::Async
 		: Name(InName)
 		, Solver(nullptr)
 		, SceneModifiers()
+		, bIsStarted(false)
 	{
 		StaticBodyPool = MakeUnique<FRigidBodyPoolAsyncPT>(InName, InStaticBodyPoolHandle);
+		RegisterBodyContainer(StaticBodyPool.Get());
+
 		DynamicBodyPool = MakeUnique<FRigidBodyPoolAsyncPT>(InName, InDynamicBodyPoolHandle);
+		RegisterBodyContainer(DynamicBodyPool.Get());
+	}
+
+	FRigidSceneAsyncPT::~FRigidSceneAsyncPT()
+	{
+		UnregisterBodyContainer(DynamicBodyPool.Get());
+		DynamicBodyPool.Reset();
+
+		UnregisterBodyContainer(StaticBodyPool.Get());
+		StaticBodyPool.Reset();
+
+		UE_RIGIDPHYSICS_CHECK(BodyContainerRegistry.IsEmpty());
 	}
 
 	void FRigidSceneAsyncPT::Lock(UE::Physics::ERigidLockType InLockType) const
@@ -459,7 +503,7 @@ namespace Chaos::Rigids::Async
 
 	void FRigidSceneAsyncPT::Tick(float Dt)
 	{
-		checkf(false, TEXT("Not supported from physics thread"));
+		UE_RIGIDPHYSICS_CHECKF(false, TEXT("Not supported from physics thread"));
 	}
 
 	UE::Physics::FRigidSceneHandle FRigidSceneAsyncPT::GetSceneHandle() const
@@ -476,6 +520,8 @@ namespace Chaos::Rigids::Async
 
 		Solver = InSolver;
 
+		bIsStarted = true;
+
 		Solver->GetEvolution()->SetPreIntegrateCallback(
 			[this](FReal Dt)
 			{
@@ -498,6 +544,41 @@ namespace Chaos::Rigids::Async
 			Solver->GetEvolution()->SetPreIntegrateCallback([](FReal Dt) {});
 			Solver->GetEvolution()->SetPostSolveCallback([](FReal Dt) {});
 		}
+
+		bIsStarted = false;
+	}
+
+	void FRigidSceneAsyncPT::RegisterBodyContainer(UE::Physics::IRigidBodyContainer* InContainer)
+	{
+		UE_RIGIDPHYSICS_CHECK(InContainer != nullptr);
+		UE_RIGIDPHYSICS_CHECK(GetBodyContainer(InContainer->GetHandle().GetId()) == nullptr);
+
+		BodyContainerRegistry.Add(InContainer->GetHandle().GetId(), InContainer);
+	}
+
+	void FRigidSceneAsyncPT::UnregisterBodyContainer(UE::Physics::IRigidBodyContainer* InContainer)
+	{
+		UE_RIGIDPHYSICS_CHECK(InContainer != nullptr);
+		UE_RIGIDPHYSICS_CHECK(GetBodyContainer(InContainer->GetHandle().GetId()) == InContainer);
+
+		BodyContainerRegistry.Remove(InContainer->GetHandle().GetId());
+	}
+
+	void FRigidSceneAsyncPT::CreateGeometryCollectionPT(const UE::Physics::FRigidBodyContainerHandle& InHandle, FGeometryCollectionPhysicsProxy* InProxy)
+	{
+		UE_RIGIDPHYSICS_CHECK(GetBodyContainer(InHandle.GetId()) == nullptr);
+
+		FRigidGeometryCollectionAsyncPT* GC = new FRigidGeometryCollectionAsyncPT(this, InHandle.GetId(), InProxy);
+		RegisterBodyContainer(GC);
+	}
+
+	void FRigidSceneAsyncPT::DestroyGeometryCollectionPT(const UE::Physics::FRigidBodyContainerHandle& InHandle)
+	{
+		if (FRigidGeometryCollectionAsyncPT* GC = GetGeometryCollection(InHandle.GetId()))
+		{
+			UnregisterBodyContainer(GC);
+			delete GC;
+		}
 	}
 
 	void FRigidSceneAsyncPT::ActivateBody(const UE::Physics::FRigidBodyHandle& InBodyHandle)
@@ -510,18 +591,18 @@ namespace Chaos::Rigids::Async
 		CheckSimContext();
 	}
 
-	void FRigidSceneAsyncPT::CreateStaticBodyPT(const UE::Physics::FRigidBodyHandle& InBodyHandle, FSingleParticlePhysicsProxy* Proxy)
+	void FRigidSceneAsyncPT::CreateStaticBodyPT(const UE::Physics::FRigidBodyHandle& InBodyHandle, FSingleParticlePhysicsProxy* InProxy)
 	{
 		CheckSimContext();
 
-		StaticBodyPool->CreateBodyPT(InBodyHandle, Proxy);
+		StaticBodyPool->CreateBodyPT(InBodyHandle, InProxy);
 	}
 
-	void FRigidSceneAsyncPT::CreateDynamicBodyPT(const UE::Physics::FRigidBodyHandle& InBodyHandle, FSingleParticlePhysicsProxy* Proxy)
+	void FRigidSceneAsyncPT::CreateDynamicBodyPT(const UE::Physics::FRigidBodyHandle& InBodyHandle, FSingleParticlePhysicsProxy* InProxy)
 	{
 		CheckSimContext();
 
-		DynamicBodyPool->CreateBodyPT(InBodyHandle, Proxy);
+		DynamicBodyPool->CreateBodyPT(InBodyHandle, InProxy);
 	}
 
 	void FRigidSceneAsyncPT::DestroyBodyPT(const UE::Physics::FRigidBodyHandle& InBodyHandle)
@@ -539,16 +620,25 @@ namespace Chaos::Rigids::Async
 
 	UE::Physics::IRigidBodyContainer* FRigidSceneAsyncPT::GetBodyContainer(const UE::Physics::FRigidObjectId& InBodyContainerId) const
 	{
+		using namespace UE::Physics;
 		CheckSimContext();
 
-		// @TODO_CHAOSAPI: Use a registry here when we need to support more containers
-		if (InBodyContainerId == StaticBodyPool->GetHandle().GetId())
+		if (IRigidBodyContainer*const* BodyContainerValuePtr = BodyContainerRegistry.Find(InBodyContainerId))
 		{
-			return StaticBodyPool.Get();
+			return *BodyContainerValuePtr;
 		}
-		else if (InBodyContainerId == DynamicBodyPool->GetHandle().GetId())
+
+		return nullptr;
+	}
+
+	FRigidGeometryCollectionAsyncPT* FRigidSceneAsyncPT::GetGeometryCollection(const UE::Physics::FRigidObjectId& InGeometryCollectionId) const
+	{
+		using namespace UE::Physics;
+		CheckSimContext();
+
+		if (IRigidBodyContainer* BodyContainer = GetBodyContainer(InGeometryCollectionId))
 		{
-			return DynamicBodyPool.Get();
+			return BodyContainer->AsA<FRigidGeometryCollectionAsyncPT>();
 		}
 
 		return nullptr;
@@ -561,7 +651,7 @@ namespace Chaos::Rigids::Async
 
 	void FRigidSceneAsyncPT::CheckSimContext() const
 	{
-		UE_RIGIDPHYSICS_CHECK(IsInPhysicsThreadContext());
+		UE_RIGIDPHYSICS_CHECK(!bIsStarted || IsInPhysicsThreadContext());
 	}
 
 	void FRigidSceneAsyncPT::OnPreTick(FReal Dt)
diff --git a/Engine/Plugins/Experimental/ChaosRigidPhysicsAsync/Source/ChaosRigidPhysicsAsync/Public/ChaosRigidPhysicsAsync/RigidFwdAsync.h b/Engine/Plugins/Experimental/ChaosRigidPhysicsAsync/Source/ChaosRigidPhysicsAsync/Public/ChaosRigidPhysicsAsync/RigidFwdAsync.h
--- a/Engine/Plugins/Experimental/ChaosRigidPhysicsAsync/Source/ChaosRigidPhysicsAsync/Public/ChaosRigidPhysicsAsync/RigidFwdAsync.h
+++ b/Engine/Plugins/Experimental/ChaosRigidPhysicsAsync/Source/ChaosRigidPhysicsAsync/Public/ChaosRigidPhysicsAsync/RigidFwdAsync.h
@@ -17,6 +17,8 @@ namespace Chaos::Rigids::Async
 	class UE_INTERNAL FRigidBodyAsyncPT;
 	class UE_INTERNAL FRigidBodyPoolAsyncGT;
 	class UE_INTERNAL FRigidBodyPoolAsyncPT;
+	class UE_INTERNAL FRigidGeometryCollectionAsyncGT;
+	class UE_INTERNAL FRigidGeometryCollectionAsyncPT;
 	class UE_INTERNAL FRigidSceneAsyncGT;
 	class UE_INTERNAL FRigidSceneAsyncPT;
 }
diff --git a/Engine/Plugins/Experimental/ChaosRigidPhysicsAsync/Source/ChaosRigidPhysicsAsync/Public/ChaosRigidPhysicsAsync/RigidGeometryCollectionAsync.h b/Engine/Plugins/Experimental/ChaosRigidPhysicsAsync/Source/ChaosRigidPhysicsAsync/Public/ChaosRigidPhysicsAsync/RigidGeometryCollectionAsync.h
new file mode 100644
--- /dev/null
+++ b/Engine/Plugins/Experimental/ChaosRigidPhysicsAsync/Source/ChaosRigidPhysicsAsync/Public/ChaosRigidPhysicsAsync/RigidGeometryCollectionAsync.h
@@ -0,0 +1,206 @@
+// Copyright Epic Games, Inc. All Rights Reserved.
+#pragma once
+
+#include "ChaosRigidPhysicsAsync/RigidFwdAsync.h"
+
+#if UE_RIGIDPHYSICS_API_ENABLED
+
+#include "RigidPhysics/RigidBody.h"
+#include "RigidPhysics/RigidGeometryCollection.h"
+
+
+namespace Chaos::Rigids::Async
+{
+	class UE_INTERNAL FRigidGeometryCollectionBodyAsyncGT : public UE::Physics::IRigidBody
+	{
+	public:
+		UE_RIGIDPHYSICS_RIGIDTYPED_DECL(CHAOSRIGIDPHYSICSASYNC_API, FRigidGeometryCollectionBodyAsyncGT);
+
+		UE_INTERNAL CHAOSRIGIDPHYSICSASYNC_API FRigidGeometryCollectionBodyAsyncGT(FRigidGeometryCollectionAsyncGT* InGeometryCollection);
+
+		virtual ~FRigidGeometryCollectionBodyAsyncGT();
+
+		//
+		// Begin IRigidBody
+		//
+		UE_INTERNAL virtual UE::Physics::FRigidBodyHandle GetHandle() const override final;
+		UE_INTERNAL virtual bool IsActive() const override final;
+		UE_INTERNAL virtual void Activate() override final;
+		UE_INTERNAL virtual void Deactivate() override final;
+		UE_INTERNAL virtual bool IsStatic() const override final;
+		UE_INTERNAL virtual bool IsKinematic() const override final;
+		UE_INTERNAL virtual bool IsDynamic() const override final;
+		UE_INTERNAL virtual bool IsSleeping() const override final;
+		UE_INTERNAL virtual int32 GetNumShapes() const override final;
+		UE_INTERNAL virtual UE::Physics::FConstRigidShapeInstancePtr GetShape(int32 InShapeIndex) const override final;
+		//UE_INTERNAL virtual void AddShape(const UE::Physics::FConstRigidShapeInstancePtr& InCollider) override final;
+		//UE_INTERNAL virtual void InitTransform(const FTransform& InTransform) override final;
+		//UE_INTERNAL virtual void UpdateTransform(const FTransform& InTransform) override final;
+		//UE_INTERNAL virtual void SetKinematicTarget(const FTransform& InTransform) override final;
+		//UE_INTERNAL virtual FTransform GetTransform() const override final;
+		//UE_INTERNAL virtual FBounds3d GetBounds() const override final;
+		//UE_INTERNAL virtual double GetMass() const override final;
+		//UE_INTERNAL virtual void SetMass(double InMass) override final;
+		//UE_INTERNAL virtual FVector3d GetInertia() const override final;
+		//UE_INTERNAL virtual void SetInertia(const FVector3d& InInertia) override final;
+		//UE_INTERNAL virtual void SetCCDEnabled(bool bInEnabled) override final {  }
+		//UE_INTERNAL virtual void SetMACDEnabled(bool bInEnabled) override final {  }
+		//UE_INTERNAL virtual void SetSmoothEdgeCollisionsEnabled(bool bInEnabled) override final {  }
+		//UE_INTERNAL virtual void SetMaxLinearVelocity(float InMax) override final {  }
+		//UE_INTERNAL virtual void SetMaxAngularVelocity(float InMax) override final {  }
+		//UE_INTERNAL virtual void SetInertiaConditioningEnabled(bool bInEnabled) override final {  }
+		//UE_INTERNAL virtual void SetGyroscopicTorqueEnabled(bool bInEnabled) override final {  }
+		//UE_INTERNAL virtual void SetGravityGroup(int32 InId) override final {  }
+		//UE_INTERNAL virtual void SetGravityScale(float InScale) override final;
+		//UE_INTERNAL virtual void SetWakeEventsEnabled(bool bInEnabled) override final {  }
+		//UE_INTERNAL virtual void SetNumPositionIterations(int32 InNum) override final {  }
+		//UE_INTERNAL virtual void SetNumVelocityIterations(int32 InNum) override final {  }
+		//UE_INTERNAL virtual void SetNumProjectionIterations(int32 InNum) override final {  }
+		//UE_INTERNAL virtual void SetIsSleeping(bool bInSleeping) override final {  }
+		//
+		// End IRigidBody
+		//
+
+	};
+
+	class UE_INTERNAL FRigidGeometryCollectionBodyAsyncPT : public UE::Physics::IRigidBody
+	{
+	public:
+		UE_RIGIDPHYSICS_RIGIDTYPED_DECL(CHAOSRIGIDPHYSICSASYNC_API, FRigidGeometryCollectionBodyAsyncPT);
+
+		UE_INTERNAL CHAOSRIGIDPHYSICSASYNC_API FRigidGeometryCollectionBodyAsyncPT(FRigidGeometryCollectionAsyncPT* InGeometryCollection);
+
+		virtual ~FRigidGeometryCollectionBodyAsyncPT();
+
+		//
+		// Begin IRigidBody
+		//
+		UE_INTERNAL virtual UE::Physics::FRigidBodyHandle GetHandle() const override final;
+		UE_INTERNAL virtual bool IsActive() const override final;
+		UE_INTERNAL virtual void Activate() override final;
+		UE_INTERNAL virtual void Deactivate() override final;
+		UE_INTERNAL virtual bool IsStatic() const override final;
+		UE_INTERNAL virtual bool IsKinematic() const override final;
+		UE_INTERNAL virtual bool IsDynamic() const override final;
+		UE_INTERNAL virtual bool IsSleeping() const override final;
+		UE_INTERNAL virtual int32 GetNumShapes() const override final;
+		UE_INTERNAL virtual UE::Physics::FConstRigidShapeInstancePtr GetShape(int32 InShapeIndex) const override final;
+		//UE_INTERNAL virtual void AddShape(const UE::Physics::FConstRigidShapeInstancePtr& InCollider) override final;
+		//UE_INTERNAL virtual void InitTransform(const FTransform& InTransform) override final;
+		//UE_INTERNAL virtual void UpdateTransform(const FTransform& InTransform) override final;
+		//UE_INTERNAL virtual void SetKinematicTarget(const FTransform& InTransform) override final;
+		//UE_INTERNAL virtual FTransform GetTransform() const override final;
+		//UE_INTERNAL virtual FBounds3d GetBounds() const override final;
+		//UE_INTERNAL virtual double GetMass() const override final;
+		//UE_INTERNAL virtual void SetMass(double InMass) override final;
+		//UE_INTERNAL virtual FVector3d GetInertia() const override final;
+		//UE_INTERNAL virtual void SetInertia(const FVector3d& InInertia) override final;
+		//UE_INTERNAL virtual void SetCCDEnabled(bool bInEnabled) override final {  }
+		//UE_INTERNAL virtual void SetMACDEnabled(bool bInEnabled) override final {  }
+		//UE_INTERNAL virtual void SetSmoothEdgeCollisionsEnabled(bool bInEnabled) override final {  }
+		//UE_INTERNAL virtual void SetMaxLinearVelocity(float InMax) override final {  }
+		//UE_INTERNAL virtual void SetMaxAngularVelocity(float InMax) override final {  }
+		//UE_INTERNAL virtual void SetInertiaConditioningEnabled(bool bInEnabled) override final {  }
+		//UE_INTERNAL virtual void SetGyroscopicTorqueEnabled(bool bInEnabled) override final {  }
+		//UE_INTERNAL virtual void SetGravityGroup(int32 InId) override final {  }
+		//UE_INTERNAL virtual void SetGravityScale(float InScale) override final;
+		//UE_INTERNAL virtual void SetWakeEventsEnabled(bool bInEnabled) override final {  }
+		//UE_INTERNAL virtual void SetNumPositionIterations(int32 InNum) override final {  }
+		//UE_INTERNAL virtual void SetNumVelocityIterations(int32 InNum) override final {  }
+		//UE_INTERNAL virtual void SetNumProjectionIterations(int32 InNum) override final {  }
+		//UE_INTERNAL virtual void SetIsSleeping(bool bInSleeping) override final {  }
+		//
+		// End IRigidBody
+		//
+
+	};
+
+	// The full rigid body API. There will be several implementations of this API. E.g., synchronous
+	// physics, asynchronous physics, immediate physics, remote physics, etc.
+	// NOTE: for some implementations like async physics, there are both Game- and Physics-Context
+	// implementations within the async implementation. Pinning a handle will return the object 
+	// appropriate for the context.
+	class UE_INTERNAL FRigidGeometryCollectionAsyncGT : public UE::Physics::IRigidGeometryCollection
+	{
+	public:
+		UE_RIGIDPHYSICS_RIGIDTYPED_DECL(CHAOSRIGIDPHYSICSASYNC_API, FRigidGeometryCollectionAsyncGT);
+
+		FRigidGeometryCollectionAsyncGT(FRigidSceneAsyncGT* InScene, const UE::Physics::FRigidDebugName& InName, const FSimulationParameters& InParameters);
+		virtual ~FRigidGeometryCollectionAsyncGT();
+
+		//
+		// Begin IRigidBodyContainer
+		//
+		UE_INTERNAL virtual UE::Physics::FRigidBodyContainerHandle GetHandle() const override final;
+		UE_INTERNAL virtual int32 GetNumBodies() const override final;
+		UE_INTERNAL virtual FRigidGeometryCollectionBodyAsyncGT* GetBody(const UE::Physics::FRigidObjectId& InBodyId) const override final;
+		UE_INTERNAL virtual void SetId(const UE::Physics::FRigidObjectId& InBodyContainerId) override final;
+		//
+		// End IRigidBodyContainer
+		//
+
+		//
+		// Begin IRigidGeometryCollection
+		//
+		UE_INTERNAL virtual FRigidGeometryCollectionBodyAsyncGT* GetBodyAt(int32 BodyIndex) const override final;
+		//
+		// End IRigidGeometryCollection
+		//
+
+
+	public:
+		UE_INTERNAL FGeometryCollectionPhysicsProxy* GetProxy() const
+		{
+			return Proxy;
+		}
+
+	private:
+		UE::Physics::FRigidDebugName Name;
+		UE::Physics::FRigidObjectId BodyContainerId;
+		TSharedPtr<FGeometryCollection> RestCollection;
+		FRigidSceneAsyncGT* Scene;
+		FGeometryCollectionPhysicsProxy* Proxy;
+	};
+
+	class UE_INTERNAL FRigidGeometryCollectionAsyncPT : public UE::Physics::IRigidGeometryCollection
+	{
+	public:
+		UE_RIGIDPHYSICS_RIGIDTYPED_DECL(CHAOSRIGIDPHYSICSASYNC_API, FRigidGeometryCollectionAsyncPT);
+
+		FRigidGeometryCollectionAsyncPT(FRigidSceneAsyncPT* InScene, const UE::Physics::FRigidObjectId& InId, FGeometryCollectionPhysicsProxy* InProxy);
+		virtual ~FRigidGeometryCollectionAsyncPT();
+
+		//
+		// Begin IRigidBodyContainer
+		//
+		UE_INTERNAL virtual UE::Physics::FRigidBodyContainerHandle GetHandle() const override final;
+		UE_INTERNAL virtual int32 GetNumBodies() const override final;
+		UE_INTERNAL virtual FRigidGeometryCollectionBodyAsyncGT* GetBody(const UE::Physics::FRigidObjectId& InBodyId) const override final;
+		UE_INTERNAL virtual void SetId(const UE::Physics::FRigidObjectId& InBodyContainerId) override final;
+		//
+		// End IRigidBodyContainer
+		//
+
+		//
+		// Begin IRigidGeometryCollection
+		//
+		UE_INTERNAL virtual FRigidGeometryCollectionBodyAsyncGT* GetBodyAt(int32 BodyIndex) const override final;
+		//
+		// End IRigidGeometryCollection
+		//
+
+	public:
+		UE_INTERNAL FGeometryCollectionPhysicsProxy* GetProxy() const
+		{
+			return Proxy;
+		}
+
+	private:
+		UE::Physics::FRigidObjectId BodyContainerId;
+		FRigidSceneAsyncPT* Scene;
+		FGeometryCollectionPhysicsProxy* Proxy;
+	};
+
+}
+
+#endif
diff --git a/Engine/Plugins/Experimental/ChaosRigidPhysicsAsync/Source/ChaosRigidPhysicsAsync/Public/ChaosRigidPhysicsAsync/RigidSceneAsync.h b/Engine/Plugins/Experimental/ChaosRigidPhysicsAsync/Source/ChaosRigidPhysicsAsync/Public/ChaosRigidPhysicsAsync/RigidSceneAsync.h
--- a/Engine/Plugins/Experimental/ChaosRigidPhysicsAsync/Source/ChaosRigidPhysicsAsync/Public/ChaosRigidPhysicsAsync/RigidSceneAsync.h
+++ b/Engine/Plugins/Experimental/ChaosRigidPhysicsAsync/Source/ChaosRigidPhysicsAsync/Public/ChaosRigidPhysicsAsync/RigidSceneAsync.h
@@ -9,6 +9,7 @@
 #include "ChaosRigidPhysicsAsync/RigidBodyAsync.h"
 #include "ChaosRigidPhysicsAsync/RigidBodyContainerAsync.h"
 #include "ChaosRigidPhysicsAsync/RigidFactoryAsync.h"
+#include "ChaosRigidPhysicsAsync/RigidGeometryCollectionAsync.h"
 #include "ChaosRigidPhysicsAsync/RigidSceneSettingsAsync.h"
 #include "ChaosSolversModule.h"
 #include "Containers/Map.h"
@@ -41,9 +42,6 @@ namespace Chaos::Rigids::Async
 	class UE_INTERNAL FRigidSceneAsyncGT 
 		: public UE::Physics::FRigidScene
 		, public TSharedFromThis<FRigidSceneAsyncGT, ESPMode::ThreadSafe>
-#if WITH_ENGINE
-		, public FGCObject
-#endif
 	{
 	public:
 		UE_RIGIDPHYSICS_RIGIDTYPED_DECL(CHAOSRIGIDPHYSICSASYNC_API, FRigidSceneAsyncGT);
@@ -53,20 +51,6 @@ namespace Chaos::Rigids::Async
 			const FRigidSceneSettingsAsync& InSceneSettings);
 		virtual ~FRigidSceneAsyncGT();
 
-#if WITH_ENGINE
-		//
-		// Begin FGCObject
-		//
-		virtual void AddReferencedObjects(FReferenceCollector& Collector) override final;
-		virtual FString GetReferencerName() const
-		{
-			return "FRigidSceneAsyncGT";
-		}
-		//
-		// End FGCObject
-		//
-#endif
-
 		//
 		// Begin IRigidLockable
 		//
@@ -88,6 +72,8 @@ namespace Chaos::Rigids::Async
 		UE_INTERNAL virtual void Tick(float Dt) override final;
 		UE_INTERNAL virtual FRigidBodyAsyncGT* CreateBody(const UE::Physics::FRigidDebugName& InName, UE::Physics::ERigidMovementType InMovementType) override final;
 		UE_INTERNAL virtual void DestroyBody(UE::Physics::IRigidBody* InBody) override final;
+		UE_INTERNAL virtual FRigidGeometryCollectionAsyncGT* CreateGeometryCollection(const UE::Physics::FRigidDebugName& InName, const FSimulationParameters& InParameters) override final;
+		UE_INTERNAL virtual void DestroyGeometryCollection(UE::Physics::IRigidGeometryCollection* InGC) override final;
 		//
 		// IRIGIDSCENE END
 		//
@@ -116,7 +102,7 @@ namespace Chaos::Rigids::Async
 		FRigidSceneSettingsAsync SceneSettings;
 
 		FChaosScene* Scene;
-		TObjectPtr<UObject> Owner;
+		TStrongObjectPtr<UObject> Owner;
 
 		TUniquePtr<FRigidBodyPoolAsyncGT> StaticBodyPool;
 		TUniquePtr<FRigidBodyPoolAsyncGT> DynamicBodyPool;
@@ -140,6 +126,8 @@ namespace Chaos::Rigids::Async
 			const UE::Physics::FRigidBodyContainerHandle& InStaticBodyPoolHandle,
 			const UE::Physics::FRigidBodyContainerHandle& InDynamicBodyPoolHandle);
 
+		CHAOSRIGIDPHYSICSASYNC_API virtual ~FRigidSceneAsyncPT();
+
 		//
 		// Begin IRigidLockable
 		//
@@ -156,6 +144,7 @@ namespace Chaos::Rigids::Async
 	public:
 		UE_INTERNAL virtual UE::Physics::FRigidSceneHandle GetHandle() const override final;
 		UE_INTERNAL virtual UE::Physics::IRigidBodyContainer* GetBodyContainer(const UE::Physics::FRigidObjectId& InBodyContainerId) const override final;
+		UE_INTERNAL virtual FRigidGeometryCollectionAsyncPT* GetGeometryCollection(const UE::Physics::FRigidObjectId& InGeometryCollectionId) const override final;
 		UE_INTERNAL virtual void RegisterModifier(UE::Physics::IRigidSceneModifier* Modifier) override final;
 		UE_INTERNAL virtual void UnregisterModifier(UE::Physics::IRigidSceneModifier* Modifier) override final;
 		UE_INTERNAL virtual bool Query(const UE::Physics::FRigidQueryDesc& QueryDesc, TArray<UE::Physics::FRigidQueryHit>& OutHits) override final;
@@ -170,8 +159,10 @@ namespace Chaos::Rigids::Async
 		UE_INTERNAL CHAOSRIGIDPHYSICSASYNC_API void Shutdown();
 		UE_INTERNAL CHAOSRIGIDPHYSICSASYNC_API void RegisterBodyContainer(UE::Physics::IRigidBodyContainer* InContainer);
 		UE_INTERNAL CHAOSRIGIDPHYSICSASYNC_API void UnregisterBodyContainer(UE::Physics::IRigidBodyContainer* InContainer);
-		UE_INTERNAL CHAOSRIGIDPHYSICSASYNC_API void CreateStaticBodyPT(const UE::Physics::FRigidBodyHandle& InBodyHandle, FSingleParticlePhysicsProxy* Proxy);
-		UE_INTERNAL CHAOSRIGIDPHYSICSASYNC_API void CreateDynamicBodyPT(const UE::Physics::FRigidBodyHandle& InBodyHandle, FSingleParticlePhysicsProxy* Proxy);
+		UE_INTERNAL CHAOSRIGIDPHYSICSASYNC_API void CreateGeometryCollectionPT(const UE::Physics::FRigidBodyContainerHandle& InHandle, FGeometryCollectionPhysicsProxy* InProxy);
+		UE_INTERNAL CHAOSRIGIDPHYSICSASYNC_API void DestroyGeometryCollectionPT(const UE::Physics::FRigidBodyContainerHandle& InHandle);
+		UE_INTERNAL CHAOSRIGIDPHYSICSASYNC_API void CreateStaticBodyPT(const UE::Physics::FRigidBodyHandle& InBodyHandle, FSingleParticlePhysicsProxy* InProxy);
+		UE_INTERNAL CHAOSRIGIDPHYSICSASYNC_API void CreateDynamicBodyPT(const UE::Physics::FRigidBodyHandle& InBodyHandle, FSingleParticlePhysicsProxy* InProxy);
 		UE_INTERNAL CHAOSRIGIDPHYSICSASYNC_API void DestroyBodyPT(const UE::Physics::FRigidBodyHandle& InBodyHandle);
 		UE_INTERNAL CHAOSRIGIDPHYSICSASYNC_API void ActivateBody(const UE::Physics::FRigidBodyHandle& InBodyHandle);
 		UE_INTERNAL CHAOSRIGIDPHYSICSASYNC_API void DeactivateBody(const UE::Physics::FRigidBodyHandle& InBodyHandle);
@@ -188,7 +179,9 @@ namespace Chaos::Rigids::Async
 		FPBDRigidsSolver* Solver;
 		TUniquePtr<FRigidBodyPoolAsyncPT> StaticBodyPool;
 		TUniquePtr<FRigidBodyPoolAsyncPT> DynamicBodyPool;
+		TMap<UE::Physics::FRigidObjectId, UE::Physics::IRigidBodyContainer*> BodyContainerRegistry;
 		TArray<UE::Physics::IRigidSceneModifier*> SceneModifiers;
+		bool bIsStarted;
 	};
 }
 
diff --git a/Engine/Plugins/Experimental/ChaosRigidPhysicsAsync/Tests/Private/GeometryCollectionTests.cpp b/Engine/Plugins/Experimental/ChaosRigidPhysicsAsync/Tests/Private/GeometryCollectionTests.cpp
new file mode 100644
--- /dev/null
+++ b/Engine/Plugins/Experimental/ChaosRigidPhysicsAsync/Tests/Private/GeometryCollectionTests.cpp
@@ -0,0 +1,147 @@
+// Copyright Epic Games, Inc. All Rights Reserved.
+
+#include "Chaos/LowLevelTest/ChaosTestHarness.h"
+#include "Chaos/LowLevelTest/ChaosTestScene.h"
+#include "ChaosRigidPhysicsAsyncTest.h"
+#include "ChaosRigidPhysicsAsync/RigidBodyContainerAsync.h"
+#include "ChaosRigidPhysicsAsync/RigidFactoryAsync.h"
+#include "ChaosRigidPhysicsAsync/RigidGeometryCollectionAsync.h"
+#include "ChaosRigidPhysicsAsync/RigidSceneAsync.h"
+#include "ChaosRigidPhysicsAsync/RigidSceneSettingsAsync.h"
+#include "RigidPhysics/RigidContext.h"
+
+#if UE_RIGIDPHYSICS_API_ENABLED
+
+namespace Chaos::LowLevelTest
+{
+	using namespace Chaos;
+	using namespace Chaos::Rigids::Async;
+	using namespace UE::Physics;
+
+	FSimulationParameters MakeGCCubeSettings()
+	{
+		// TODO_CHAOSAPI: A lot of code required here
+		return FSimulationParameters();
+	}
+
+
+
+	// Test casting between GC and BodyPool (both types of body container).
+	// Casts from BodyContainer to/from GC should compile, but fail gracefully when the
+	// types don't match.
+	TEST_CASE("GeometryCollectionCastingTests", "[Chaos][API][GeometryCollection][unit]")
+	{
+		FRigidFactoryAsync Factory;
+		FRigidSceneSettingsAsync SceneSettings;
+		FRigidSceneHandle SceneHandle = Factory.CreateScene(FRigidDebugName::Make(TEXT("Main")), &SceneSettings);
+
+		// We can safely (fail to) cast an empty body container to a GC, giving a null GC ptr
+		if (FRigidContextGameRW Context = SceneHandle.LockRW())
+		{
+			FRigidBodyContainerHandle BodyContainerHandle;
+			CHECK_FALSE(BodyContainerHandle.Pin(Context).IsValid());
+			TRigidGeometryCollectionPtr<FRigidContextGameRW> GCPtr = BodyContainerHandle.Pin(Context);
+			CHECK_FALSE(GCPtr.IsValid());
+		}
+
+		// We can safely (fail to) cast a valid body container to a GC, giving a null GC ptr
+		if (FRigidContextGameRW Context = SceneHandle.LockRW())
+		{
+			// Create a standalone body (i.e., not in a GC)
+			TRigidBodyPtr<FRigidContextGameRW> BodyPtr = Context->CreateBody(FRigidDebugName::Make(TEXT("Body1")), ERigidMovementType::Static);
+			REQUIRE(BodyPtr.IsValid());
+
+			// Get the container that owns the body - it should be a body pool, not a GC
+			TRigidBodyContainerPtr<FRigidContextGameRW> BodyContainerPtr = Context->GetBodyContainer(BodyPtr->GetContainerHandle().GetId());
+			CHECK(BodyContainerPtr.IsValid());
+			CHECK(BodyContainerPtr->GetTypeId() == FRigidBodyPoolAsyncGT::GetStaticTypeId());
+
+			// Attempt to cast the container to a GC - should fail gracefully
+			TRigidGeometryCollectionPtr<FRigidContextGameRW> GCPtr = BodyContainerPtr;
+			CHECK_FALSE(GCPtr.IsValid());
+
+			Context->DestroyBody(MoveTemp(BodyPtr));
+		}
+
+		// We can cast a valid GCPtr to a valid BodyContainerPtr
+		if (FRigidContextGameRW Context = SceneHandle.LockRW())
+		{
+			// TODO_CHAOSAPI: GeometryCollection initialization and test objects
+			FSimulationParameters GCSettings = MakeGCCubeSettings();
+			if (TRigidGeometryCollectionPtr<FRigidContextGameRW> GCPtr = Context->CreateGeometryCollection(FRigidDebugName::Make(TEXT("GC1")), GCSettings))
+			{
+				CHECK(GCPtr->GetTypeId() == FRigidGeometryCollectionAsyncGT::GetStaticTypeId());
+
+				// Create a GC Handle
+				FRigidGeometryCollectionHandle GCHandle = GCPtr;
+				CHECK(GCHandle.Pin(Context).IsValid());
+
+				// The GC handle should also be a valid BodyContainer handle
+				FRigidBodyContainerHandle BCHandle = GCHandle;
+				TRigidBodyContainerPtr<FRigidContextGameRW> BCPtr2 = BCHandle.Pin(Context);
+				CHECK(BCPtr2.IsValid());
+				CHECK(BCPtr2->GetTypeId() == FRigidGeometryCollectionAsyncGT::GetStaticTypeId());
+
+				// Recover the GCPtr
+				TRigidGeometryCollectionPtr<FRigidContextGameRW> GCPtr2 = BCPtr2;
+				CHECK(GCPtr2.IsValid());
+				CHECK(GCPtr2->GetTypeId() == FRigidGeometryCollectionAsyncGT::GetStaticTypeId());
+
+				Context->DestroyGeometryCollection(MoveTemp(GCPtr));
+			}
+		}
+
+
+		Factory.DestroyScene(SceneHandle);
+	}
+
+
+	// Create A GC with some bodies in it and verify that we can access them.
+	// TODO_CHAOSPI: This will fail until we add code to generate a GC with at least one body in it
+	TEST_CASE("GeometryCollectionInitializationTest", "[!shouldfail][Chaos][API][GeometryCollection][unit]")
+	{
+		FRigidFactoryAsync Factory;
+		FRigidSceneSettingsAsync SceneSettings;
+		FRigidSceneHandle SceneHandle = Factory.CreateScene(FRigidDebugName::Make(TEXT("Main")), &SceneSettings);
+
+		if (FRigidContextGameRW Context = SceneHandle.LockRW())
+		{
+			FSimulationParameters GCSettings = MakeGCCubeSettings();
+			TRigidGeometryCollectionPtr<FRigidContextGameRW> GCPtr = Context->CreateGeometryCollection(FRigidDebugName::Make(TEXT("GC1")), GCSettings);
+			REQUIRE(GCPtr.IsValid());
+
+			// We can recover the GC from its handle on the GT
+			{
+				FRigidGeometryCollectionHandle GCHandle = GCPtr;
+				TRigidGeometryCollectionPtr<FRigidContextGameRW> GCPtr2 = Context->GetGeometryCollection(GCHandle.GetId());
+				CHECK(GCPtr2.IsValid());
+				CHECK(GCPtr2.Get() == GCPtr.Get());
+			}
+
+			// We can access the GC as a BodyContainer on the GT
+			{
+				FRigidGeometryCollectionHandle GCHandle = GCPtr;
+				TRigidBodyContainerPtr<FRigidContextGameRW> BodyContainerPtr = Context->GetBodyContainer(GCHandle.GetId());
+				CHECK(BodyContainerPtr.IsValid());
+				CHECK(BodyContainerPtr.Get() == GCPtr.Get());
+			}
+
+			// The GC has at least one body
+			CHECK(GCPtr->GetNumBodies() > 0);
+			TRigidBodyPtr<FRigidContextGameRW> BodyPtr = GCPtr->GetBodyAt(0);
+			CHECK(BodyPtr.IsValid());
+
+			// A body from the GC should be a GC Body
+			if (BodyPtr.IsValid())
+			{
+				CHECK(BodyPtr->GetTypeId() == FRigidGeometryCollectionBodyAsyncGT::GetStaticTypeId());
+			}
+
+			Context->DestroyGeometryCollection(MoveTemp(GCPtr));
+		}
+
+		Factory.DestroyScene(SceneHandle);
+	}
+}
+
+#endif
\ No newline at end of file
diff --git a/Engine/Source/Runtime/Experimental/RigidPhysics/Private/RigidPhysics/RigidGeometryCollection.cpp b/Engine/Source/Runtime/Experimental/RigidPhysics/Private/RigidPhysics/RigidGeometryCollection.cpp
new file mode 100644
--- /dev/null
+++ b/Engine/Source/Runtime/Experimental/RigidPhysics/Private/RigidPhysics/RigidGeometryCollection.cpp
@@ -0,0 +1,11 @@
+// Copyright Epic Games, Inc. All Rights Reserved.
+#include "RigidPhysics/RigidGeometryCollection.h"
+
+#if UE_RIGIDPHYSICS_API_ENABLED
+
+namespace UE::Physics
+{
+	UE_RIGIDPHYSICS_RIGIDTYPED_IMPL(IRigidGeometryCollection, IRigidTyped);
+}
+
+#endif // UE_RIGIDPHYSICS_API_ENABLED
\ No newline at end of file
diff --git a/Engine/Source/Runtime/Experimental/RigidPhysics/Private/RigidPhysics/RigidScene.cpp b/Engine/Source/Runtime/Experimental/RigidPhysics/Private/RigidPhysics/RigidScene.cpp
--- a/Engine/Source/Runtime/Experimental/RigidPhysics/Private/RigidPhysics/RigidScene.cpp
+++ b/Engine/Source/Runtime/Experimental/RigidPhysics/Private/RigidPhysics/RigidScene.cpp
@@ -3,6 +3,7 @@
 #include "RigidPhysics/RigidBody.h"
 #include "RigidPhysics/RigidBodyContainer.h"
 #include "RigidPhysics/RigidContext.h"
+#include "RigidPhysics/RigidGeometryCollection.h"
 #include "RigidPhysics/RigidLog.h"
 #include "RigidPhysics/RigidSceneSettings.h"
 
@@ -63,6 +64,16 @@ namespace UE::Physics
 	{
 		return BodyContainerRegistry->FindObject(InBodyContainerId);
 	}
+
+	IRigidGeometryCollection* FRigidScene::GetGeometryCollection(const FRigidObjectId& InGeometryCollectionId) const
+	{
+		if (IRigidBodyContainer* BodyContainer = BodyContainerRegistry->FindObject(InGeometryCollectionId))
+		{
+			return BodyContainer->AsA<IRigidGeometryCollection>();
+		}
+		return nullptr;
+	}
+
 }
 
 #endif // UE_RIGIDPHYSICS_API_ENABLED
\ No newline at end of file
diff --git a/Engine/Source/Runtime/Experimental/RigidPhysics/Public/RigidPhysics/RigidBody.h b/Engine/Source/Runtime/Experimental/RigidPhysics/Public/RigidPhysics/RigidBody.h
--- a/Engine/Source/Runtime/Experimental/RigidPhysics/Public/RigidPhysics/RigidBody.h
+++ b/Engine/Source/Runtime/Experimental/RigidPhysics/Public/RigidPhysics/RigidBody.h
@@ -275,10 +275,7 @@ namespace UE::Physics
 
 		FRigidBodyHandle GetHandle() const;
 
-		FRigidBodyContainerHandle GetContainerHandle() const
-		{
-			return BodyRaw->GetContainerHandle();
-		}
+		FRigidBodyContainerHandle GetContainerHandle() const;
 
 		bool IsActive() const
 		{
@@ -507,7 +504,17 @@ namespace UE::Physics
 	template<typename ContextType>
 	FRigidBodyHandle TRigidBodyPtrImpl<ContextType>::GetHandle() const
 	{
-		return BodyRaw->GetHandle();
+		if (BodyRaw != nullptr)
+		{
+			return BodyRaw->GetHandle();
+		}
+		return FRigidBodyHandle();
+	}
+
+	template<typename ContextType>
+	FRigidBodyContainerHandle TRigidBodyPtrImpl<ContextType>::GetContainerHandle() const
+	{
+		return GetHandle().GetContainerHandle();
 	}
 
 	//
@@ -515,9 +522,9 @@ namespace UE::Physics
 	//
 
 	template<typename ContextType>
-	TRigidBodyPtr<ContextType> TRigidBodyContainerPtrImpl<ContextType>::GetBody(const FRigidObjectId& BodyId) const
+	TRigidBodyPtr<ContextType> TRigidBodyContainerPtrImpl<ContextType>::GetBody(const FRigidObjectId& InBodyId) const
 	{
-		return BodyContainerRaw->GetBody(BodyId);
+		return BodyContainerRaw->GetBody(InBodyId);
 	}
 
 	template<typename ContextType>
diff --git a/Engine/Source/Runtime/Experimental/RigidPhysics/Public/RigidPhysics/RigidBodyContainer.h b/Engine/Source/Runtime/Experimental/RigidPhysics/Public/RigidPhysics/RigidBodyContainer.h
--- a/Engine/Source/Runtime/Experimental/RigidPhysics/Public/RigidPhysics/RigidBodyContainer.h
+++ b/Engine/Source/Runtime/Experimental/RigidPhysics/Public/RigidPhysics/RigidBodyContainer.h
@@ -132,6 +132,11 @@ namespace UE::Physics
 		TRigidBodyPtr<FContext> CreateBody(const FRigidDebugName& InName, ERigidMovementType InMovementType) const;
 		void DestroyBody(TRigidBodyPtr<FContext>&& InBodyPtr);
 
+		UE_INTERNAL IRigidBodyContainer* Get() const
+		{
+			return BodyContainerRaw;
+		}
+
 	private:
 		// Pointer to implementation, object depends on Context type
 		IRigidBodyContainer* BodyContainerRaw;
@@ -162,6 +167,14 @@ namespace UE::Physics
 			return BodyContainerId;
 		}
 
+		// NOTE: We don't really need the scene handle because it is provided by the Context
+		// in Pin(). It is only needed for the debugger visualizers (see RigidPhysics.natvis).
+		// TODO_CHAOSAPI: Consider making SceneHandle shipping only
+		const FRigidSceneHandle GetSceneHandle() const
+		{
+			return SceneHandle;
+		}
+
 		// Convert the handle to a pointer in the specified context
 		// Impl in RigidContext.h
 		template<typename ContextType>
@@ -172,11 +185,6 @@ namespace UE::Physics
 
 	private:
 		FRigidObjectId BodyContainerId;
-
-		// NOTE: This is deliberately not exposed through a GetSceneHandle() method
-		// We don't really need the scene handle because it is provided by the Context
-		// in Pin(). It is only needed for the debugger visualizers.
-		// TODO_CHAOSAPI: Make this shipping only?
 		FRigidSceneHandle SceneHandle;
 	};
 
diff --git a/Engine/Source/Runtime/Experimental/RigidPhysics/Public/RigidPhysics/RigidContext.h b/Engine/Source/Runtime/Experimental/RigidPhysics/Public/RigidPhysics/RigidContext.h
--- a/Engine/Source/Runtime/Experimental/RigidPhysics/Public/RigidPhysics/RigidContext.h
+++ b/Engine/Source/Runtime/Experimental/RigidPhysics/Public/RigidPhysics/RigidContext.h
@@ -9,6 +9,7 @@
 #include "RigidPhysics/RigidLog.h"
 #include "RigidPhysics/RigidBody.h"
 #include "RigidPhysics/RigidBodyContainer.h"
+#include "RigidPhysics/RigidGeometryCollection.h"
 #include "RigidPhysics/RigidScene.h"
 
 namespace UE::Physics
@@ -97,7 +98,19 @@ namespace UE::Physics
 		// NOTE: we don't use our copy of SceneHandle here (or anywhere!)
 		if (InContext.IsValid())
 		{
-			return InContext.GetScene()->GetBodyContainer(BodyContainerId);
+			return InContext.GetScene()->GetBodyContainer(GetId());
+		}
+		return nullptr;
+	}
+
+	template<typename ContextType>
+	inline TRigidGeometryCollectionPtr<ContextType> FRigidGeometryCollectionHandle::Pin(const ContextType& InContext) const
+	{
+		if (InContext.IsValid())
+		{
+			// TODO_CHAOSAPI: Could use checked cast here because a FRigidGeometryCollectionHandle will 
+			// always reference a GC or nothing. Currently it will perform an RTTI cast.
+			return BodyContainerHandle.Pin(InContext);
 		}
 		return nullptr;
 	}
@@ -107,9 +120,9 @@ namespace UE::Physics
 	{
 		if (InContext.IsValid())
 		{
-			if (TRigidBodyContainerPtr<ContextType> BodyContainer = InContext.GetScene()->GetBodyContainer(BodyContainerHandle.GetId()))
+			if (TRigidBodyContainerPtr<ContextType> BodyContainer = BodyContainerHandle.Pin(InContext))
 			{
-				return BodyContainer->GetBody(BodyId);
+				return BodyContainer->GetBody(GetId());
 			}
 		}
 		return TRigidBodyPtr<ContextType>();
diff --git a/Engine/Source/Runtime/Experimental/RigidPhysics/Public/RigidPhysics/RigidFwd.h b/Engine/Source/Runtime/Experimental/RigidPhysics/Public/RigidPhysics/RigidFwd.h
--- a/Engine/Source/Runtime/Experimental/RigidPhysics/Public/RigidPhysics/RigidFwd.h
+++ b/Engine/Source/Runtime/Experimental/RigidPhysics/Public/RigidPhysics/RigidFwd.h
@@ -49,11 +49,15 @@ template<typename T> struct TBounds;
 using FBounds3f = TBounds<float>;
 using FBounds3d = TBounds<double>;
 
+class FGeometryCollection;
+struct FSimulationParameters;	// GC setup
+
 // User-facing classes
 namespace UE::Physics
 {
 	class FRigidBodyContainerHandle;
 	class FRigidBodyHandle;
+	class FRigidGeometryCollectionHandle;
 	class FRigidObjectId;
 	class FRigidQueryDesc;
 	class FRigidQueryHit;
@@ -62,6 +66,7 @@ namespace UE::Physics
 
 	class IRigidBody;
 	class IRigidBodyContainer;
+	class IRigidGeometryCollection;
 	class IRigidLockable;
 	class IRigidQuery;
 	class IRigidScene;
@@ -72,12 +77,14 @@ namespace UE::Physics
 
 	template<typename ContextType> class TRigidBodyContainerPtrImpl;
 	template<typename ContextType> class TRigidBodyPtrImpl;
-	template<bool bConstParam> class TRigidShapeInstancePtr;
+	template<typename ContextType> class TRigidGeometryCollectionPtrImpl;
 	template<typename PtrImplType> class TRigidObjectPtr;
 	template<typename ContextType> class TRigidScenePtrImpl;
+	template<bool bConstParam> class TRigidShapeInstancePtr;
 
 	template<typename ContextType> using TRigidBodyContainerPtr = TRigidObjectPtr<TRigidBodyContainerPtrImpl<ContextType>>;
 	template<typename ContextType> using TRigidBodyPtr = TRigidObjectPtr<TRigidBodyPtrImpl<ContextType>>;
+	template<typename ContextType> using TRigidGeometryCollectionPtr = TRigidObjectPtr<TRigidGeometryCollectionPtrImpl<ContextType>>;
 	template<typename ContextType> using TRigidScenePtr = TRigidObjectPtr<TRigidScenePtrImpl<ContextType>>;
 
 	using FConstRigidShapeInstancePtr = TRigidShapeInstancePtr<true>;
@@ -132,6 +139,7 @@ namespace UE::Physics
 {
 	template<typename ObjectType> class UE_INTERNAL TRigidObjectRegistry;
 	using UE_INTERNAL FRigidBodyContainerRegistry = TRigidObjectRegistry<IRigidBodyContainer>;
+	using UE_INTERNAL FRigidGeometryCollectionRegistry = TRigidObjectRegistry<IRigidGeometryCollection>;
 	using UE_INTERNAL FRigidBodyRegistry = TRigidObjectRegistry<IRigidBody>;
 }
 
diff --git a/Engine/Source/Runtime/Experimental/RigidPhysics/Public/RigidPhysics/RigidGeometryCollection.h b/Engine/Source/Runtime/Experimental/RigidPhysics/Public/RigidPhysics/RigidGeometryCollection.h
new file mode 100644
--- /dev/null
+++ b/Engine/Source/Runtime/Experimental/RigidPhysics/Public/RigidPhysics/RigidGeometryCollection.h
@@ -0,0 +1,212 @@
+// Copyright Epic Games, Inc. All Rights Reserved.
+#pragma once
+
+#include "RigidPhysics/RigidFwd.h"
+
+#if UE_RIGIDPHYSICS_API_ENABLED
+
+#include "RigidPhysics/RigidBody.h"
+#include "RigidPhysics/RigidBodyContainer.h"
+
+
+namespace UE::Physics
+{
+	// The interface for rigid body containers. There will be several implementations of this API. E.g., synchronous
+	// physics, asynchronous physics, immediate physics, remote physics, etc. 
+	class UE_INTERNAL IRigidGeometryCollection : public IRigidBodyContainer
+	{
+	public:
+		UE_RIGIDPHYSICS_RIGIDTYPED_DECL(RIGIDPHYSICS_API, IRigidGeometryCollection);
+
+		UE_INTERNAL IRigidGeometryCollection()
+		{
+		}
+
+		UE_INTERNAL virtual ~IRigidGeometryCollection()
+		{
+		}
+
+		// TODO_CHAOSAPI: Figure out the public GeometryCollection API!!
+		UE_INTERNAL RIGIDPHYSICS_API virtual IRigidBody* GetBodyAt(int32 BodyIndex) const = 0;
+	};
+
+	// Utility wrapper around IRigidGeometryCollection that adds Context and allows
+	// TRigidGeometryCollectionPtr to have pointer semantics without exposing IRigidGeometryCollection
+	template<typename ContextType>
+	class UE_INTERNAL TRigidGeometryCollectionPtrImpl
+	{
+	public:
+		using FContext = ContextType;
+		using FInterface = IRigidGeometryCollection;
+		using FHandle = FRigidGeometryCollectionHandle;
+		using FRigidScenePtr = TRigidScenePtr<FContext>;
+
+		UE_INTERNAL TRigidGeometryCollectionPtrImpl(IRigidGeometryCollection* InGeometryCollection)
+			: GeometryCollectionRaw(InGeometryCollection)
+		{
+		}
+
+		UE_INTERNAL TRigidGeometryCollectionPtrImpl(IRigidBodyContainer* InBodyContainer)
+			: GeometryCollectionRaw(nullptr)
+		{
+			if (InBodyContainer != nullptr)
+			{
+				GeometryCollectionRaw = InBodyContainer->AsA<IRigidGeometryCollection>();
+			}
+		}
+
+		virtual ~TRigidGeometryCollectionPtrImpl()
+		{
+		}
+
+		UE_INTERNAL void Reset()
+		{
+			GeometryCollectionRaw = nullptr;
+		}
+
+		UE_INTERNAL bool IsValid() const
+		{
+			return (GeometryCollectionRaw != nullptr);
+		}
+
+		const FRigidTypeId& GetTypeId() const
+		{
+			return GeometryCollectionRaw->GetTypeId();
+		}
+
+		FRigidObjectId GetId() const
+		{
+			return GeometryCollectionRaw->GetId();
+		}
+
+		FRigidGeometryCollectionHandle GetHandle() const;
+
+		bool IsInScene() const
+		{
+			return GeometryCollectionRaw->IsInScene();
+		}
+
+		int32 GetNumBodies() const
+		{
+			return GeometryCollectionRaw->GetNumBodies();
+		}
+
+		// Impl in RigidBody.h
+		TRigidBodyPtr<FContext> GetBodyAt(int32 InBodyIndex) const;
+		TRigidBodyPtr<FContext> GetBody(const FRigidObjectId& InBodyId) const;
+
+		UE_INTERNAL IRigidGeometryCollection* Get() const
+		{
+			return GeometryCollectionRaw;
+		}
+
+	private:
+		// Pointer to implementation, object depends on Context type
+		IRigidGeometryCollection* GeometryCollectionRaw;
+	};
+
+	// A handle to a GeometryCollection. A GeometryCollection is a type of BodyContainer and the handle
+	// is auto-castable to FRigidBodyContainerHandle, but the reverse is more complex.
+	//
+	// NOTE: If the GeometryCollectionHandle has a valid BodyContainerHandle, it will always reference a GC. 
+	// There should be no way to get a GCHandle that contains a non-GC body container handle. This is enforced
+	// because the only way to create a FRigidGeometryCollectionHandle is from a TRigidGeometryCollectionPtr.
+	// 
+	// To safely convert a FRigidBodyContainerHandle to FRigidGeometryCollectionHandle we must go through
+	// pointer conversion:
+	// 
+	//		FRigidBodyContainerHandle BodyContainerHandle = ...;
+	//		TRigidGeomatryCollectionPtr<FRigidContextGameRW> GCPtr = BodyContainerHandle.Pin(Context);
+	//		FRigidGeometryCollectionHandle GCHandle = GCPtr;
+	// 
+	// This ensures that we get an empty handle if BodyContainerHandle is not a GC, because the cast from
+	// TRigidBodyContainerPtr to TRigidGeometryCollectionPtr will fail and return nullptr.
+	//
+	class FRigidGeometryCollectionHandle
+	{
+	public:
+		FRigidGeometryCollectionHandle()
+		{
+		}
+
+		UE_INTERNAL FRigidGeometryCollectionHandle(const FRigidSceneHandle& InSceneHandle, const FRigidObjectId& InBodyContainerId)
+			: BodyContainerHandle(InSceneHandle, InBodyContainerId)
+		{
+		}
+
+		void Reset()
+		{
+			BodyContainerHandle.Reset();
+		}
+
+		const FRigidObjectId& GetId() const
+		{
+			return BodyContainerHandle.GetId();
+		}
+
+		const FRigidSceneHandle GetSceneHandle() const
+		{
+			return BodyContainerHandle.GetSceneHandle();
+		}
+
+		FRigidBodyContainerHandle AsBodyContainerHandle() const
+		{
+			return BodyContainerHandle;
+		}
+
+		operator FRigidBodyContainerHandle() const
+		{
+			return AsBodyContainerHandle();
+		}
+
+		// Convert the handle to a pointer in the specified context
+		// Impl in RigidContext.h
+		template<typename ContextType>
+		TRigidGeometryCollectionPtr<ContextType> Pin(const ContextType& InContext) const;
+
+		friend bool operator==(const FRigidBodyContainerHandle& L, const FRigidBodyContainerHandle& R) = default;
+		friend bool operator!=(const FRigidBodyContainerHandle& L, const FRigidBodyContainerHandle& R) = default;
+
+	private:
+		FRigidBodyContainerHandle BodyContainerHandle;
+	};
+
+
+	template<typename ContextType>
+	FRigidGeometryCollectionHandle TRigidGeometryCollectionPtrImpl<ContextType>::GetHandle() const
+	{
+		if (GeometryCollectionRaw != nullptr)
+		{
+			FRigidBodyContainerHandle BodyContainerHandle = GeometryCollectionRaw->GetHandle();
+			return FRigidGeometryCollectionHandle(BodyContainerHandle.GetSceneHandle(), BodyContainerHandle.GetId());
+		}
+		return FRigidGeometryCollectionHandle();
+	}
+
+	template<typename ContextType>
+	TRigidBodyPtr<ContextType> TRigidGeometryCollectionPtrImpl<ContextType>::GetBodyAt(int32 InBodyIndex) const
+	{
+		return GeometryCollectionRaw->GetBodyAt(InBodyIndex);
+	}
+
+	template<typename ContextType>
+	TRigidBodyPtr<ContextType> TRigidGeometryCollectionPtrImpl<ContextType>::GetBody(const FRigidObjectId& InBodyId) const
+	{
+		return GeometryCollectionRaw->GetBody(InBodyId);
+	}
+
+	template<typename ContextType>
+	TRigidGeometryCollectionPtr<ContextType> TRigidScenePtrImpl<ContextType>::CreateGeometryCollection(const FRigidDebugName& InName, const FSimulationParameters& InParameters)
+	{
+		return SceneRaw->CreateGeometryCollection(InName, InParameters);
+	}
+
+	template<typename ContextType>
+	void TRigidScenePtrImpl<ContextType>::DestroyGeometryCollection(TRigidGeometryCollectionPtr<ContextType>&& InGC)
+	{
+		SceneRaw->DestroyGeometryCollection(InGC.Get());
+		InGC.Reset();
+	}
+}
+
+#endif // UE_RIGIDPHYSICS_API_ENABLED
\ No newline at end of file
diff --git a/Engine/Source/Runtime/Experimental/RigidPhysics/Public/RigidPhysics/RigidObjectPtr.h b/Engine/Source/Runtime/Experimental/RigidPhysics/Public/RigidPhysics/RigidObjectPtr.h
--- a/Engine/Source/Runtime/Experimental/RigidPhysics/Public/RigidPhysics/RigidObjectPtr.h
+++ b/Engine/Source/Runtime/Experimental/RigidPhysics/Public/RigidPhysics/RigidObjectPtr.h
@@ -5,6 +5,8 @@
 
 #if UE_RIGIDPHYSICS_API_ENABLED
 
+#include "Templates/Requires.h"
+
 namespace UE::Physics
 {
 
@@ -44,6 +46,17 @@ namespace UE::Physics
 		{
 		}
 
+		// Downcasting. Will fail to compile for attempts to cast to non-derived types.
+		// Will generate "invalid" pointer in the runtime when InPtr is not an instance of right type.
+		template<
+			typename OtherPointerImplType
+			UE_REQUIRES(std::is_base_of_v<typename OtherPointerImplType::FInterface, FInterface>)
+			>
+		UE_INTERNAL TRigidObjectPtr(const TRigidObjectPtr<OtherPointerImplType>& InPtr) 
+			: Impl(InPtr.Get())
+		{
+		}
+
 		~TRigidObjectPtr()
 		{
 		}
@@ -98,7 +111,6 @@ namespace UE::Physics
 		// Extract the object interface (internal use only)
 		UE_INTERNAL FInterface* Get() const
 		{
-			CheckIsValid();
 			return Impl.Get();
 		}
 
diff --git a/Engine/Source/Runtime/Experimental/RigidPhysics/Public/RigidPhysics/RigidScene.h b/Engine/Source/Runtime/Experimental/RigidPhysics/Public/RigidPhysics/RigidScene.h
--- a/Engine/Source/Runtime/Experimental/RigidPhysics/Public/RigidPhysics/RigidScene.h
+++ b/Engine/Source/Runtime/Experimental/RigidPhysics/Public/RigidPhysics/RigidScene.h
@@ -12,6 +12,7 @@
 #include "RigidPhysics/RigidTyped.h"
 #include "Templates/SharedPointer.h"
 
+
 namespace UE::Physics
 {
 	class IRigidSceneFactory
@@ -48,6 +49,8 @@ namespace UE::Physics
 
 		UE_INTERNAL RIGIDPHYSICS_API virtual IRigidBodyContainer* GetBodyContainer(const FRigidObjectId& InBodyContainerId) const = 0;
 
+		UE_INTERNAL RIGIDPHYSICS_API virtual IRigidGeometryCollection* GetGeometryCollection(const FRigidObjectId& InGeometryCollectionId) const = 0;
+
 		UE_INTERNAL RIGIDPHYSICS_API virtual bool Query(const FRigidQueryDesc& InQueryDesc, TArray<FRigidQueryHit>& OutHits)
 		{
 			RIGIDPHYSICS_API_UNSUPPORTED();
@@ -82,6 +85,17 @@ namespace UE::Physics
 			RIGIDPHYSICS_API_UNSUPPORTED();
 		}
 
+		UE_INTERNAL RIGIDPHYSICS_API virtual IRigidGeometryCollection* CreateGeometryCollection(const FRigidDebugName& InName, const FSimulationParameters& InParameters)
+		{
+			RIGIDPHYSICS_API_UNSUPPORTED();
+			return nullptr;
+		}
+
+		UE_INTERNAL RIGIDPHYSICS_API virtual void DestroyGeometryCollection(IRigidGeometryCollection* InGC)
+		{
+			RIGIDPHYSICS_API_UNSUPPORTED();
+		}
+
 	};
 
 	//
@@ -95,6 +109,7 @@ namespace UE::Physics
 		RIGIDPHYSICS_API virtual ~FRigidScene();
 
 		UE_INTERNAL RIGIDPHYSICS_API virtual IRigidBodyContainer* GetBodyContainer(const FRigidObjectId& InBodyContainerId) const override final;
+		UE_INTERNAL RIGIDPHYSICS_API virtual IRigidGeometryCollection* GetGeometryCollection(const FRigidObjectId& InGeometryCollectionId) const override final;
 
 	protected:
 		UE_INTERNAL RIGIDPHYSICS_API FRigidBodyContainerRegistry* GetBodyContainerRegistry() const
@@ -149,6 +164,11 @@ namespace UE::Physics
 			return SceneRaw->GetBodyContainer(InBodyContainerId);
 		}
 
+		IRigidGeometryCollection* GetGeometryCollection(const FRigidObjectId& InGeometryCollectionId) const
+		{
+			return SceneRaw->GetGeometryCollection(InGeometryCollectionId);
+		}
+
 		void RegisterModifier(IRigidSceneModifier* Modifier)
 		{
 			SceneRaw->RegisterModifier(Modifier);
@@ -163,6 +183,10 @@ namespace UE::Physics
 		TRigidBodyPtr<FContext> CreateBody(const FRigidDebugName& InName, ERigidMovementType InMovementType);
 		void DestroyBody(TRigidBodyPtr<FContext>&& InBody);
 
+		// Impl in RigidGeometryCollection.h
+		TRigidGeometryCollectionPtr<FContext> CreateGeometryCollection(const FRigidDebugName& InName, const FSimulationParameters& InParameters);
+		void DestroyGeometryCollection(TRigidGeometryCollectionPtr<ContextType>&& InGC);
+
 		bool Query(const FRigidQueryDesc& QueryDesc, TArray<FRigidQueryHit>& OutHits)  const
 		{
 			return SceneRaw->Query(QueryDesc, OutHits);
