Chaos API: Detaching RigidBody from RigidBodyContainer.

Now body handles directly contain the scene handle instead of a body collection and there's no way to go between them. BodyContainer still exists for now, but will be removed at a later point.

#Synced-CL 50971911

diff --git a/Engine/Plugins/Experimental/ChaosRigidPhysicsAsync/Source/ChaosRigidPhysicsAsync/Private/ChaosRigidPhysicsAsync/RigidBodyAsync.cpp b/Engine/Plugins/Experimental/ChaosRigidPhysicsAsync/Source/ChaosRigidPhysicsAsync/Private/ChaosRigidPhysicsAsync/RigidBodyAsync.cpp
--- a/Engine/Plugins/Experimental/ChaosRigidPhysicsAsync/Source/ChaosRigidPhysicsAsync/Private/ChaosRigidPhysicsAsync/RigidBodyAsync.cpp
+++ b/Engine/Plugins/Experimental/ChaosRigidPhysicsAsync/Source/ChaosRigidPhysicsAsync/Private/ChaosRigidPhysicsAsync/RigidBodyAsync.cpp
@@ -1,7 +1,6 @@
 // Copyright Epic Games, Inc. All Rights Reserved.
 #include "ChaosRigidPhysicsAsync/RigidBodyAsync.h"
 
-#include "ChaosRigidPhysicsAsync/RigidBodyContainerAsync.h"
 #include "ChaosRigidPhysicsAsync/RigidSceneAsync.h"
 #include "Chaos/ShapeInstance.h"
 #include "RigidPhysics/RigidBody.h"
@@ -30,13 +29,13 @@ namespace Chaos::Rigids::Async
 			SafeInvert(InValue.Z, InDefaultInverse));
 	}
 
-	FRigidBodyAsyncGT::FRigidBodyAsyncGT(FRigidBodyPoolAsyncGT* InContainer, const UE::Physics::FRigidDebugName& InName, UE::Physics::ERigidMovementType InMovementType)
+	FRigidBodyAsyncGT::FRigidBodyAsyncGT(FRigidSceneAsyncGT* InScene, const UE::Physics::FRigidDebugName& InName, UE::Physics::ERigidMovementType InMovementType)
 		: BodyId()
-		, BodyContainer(InContainer)
+		, SceneAsync(InScene)
 		, Proxy(nullptr)
 	{
 		using namespace UE::Physics;
-		UE_RIGIDPHYSICS_CHECK(BodyContainer != nullptr);
+		UE_RIGIDPHYSICS_CHECK(SceneAsync != nullptr);
 
 		TUniquePtr<FGeometryParticle> Particle = nullptr;
 		if (InMovementType == ERigidMovementType::Static)
@@ -78,7 +77,7 @@ namespace Chaos::Rigids::Async
 	{
 		using namespace UE::Physics;
 
-		return FRigidBodyHandle(BodyContainer->GetHandle(), BodyId);
+		return FRigidBodyHandle(GetScene()->GetHandle(), BodyId);
 	}
 
 	void FRigidBodyAsyncGT::SetId(const UE::Physics::FRigidObjectId& InId)
@@ -115,8 +114,7 @@ namespace Chaos::Rigids::Async
 
 	FRigidSceneAsyncGT* FRigidBodyAsyncGT::GetScene() const
 	{
-		UE_RIGIDPHYSICS_CHECK(BodyContainer != nullptr);
-		return BodyContainer->GetScene();
+		return SceneAsync;
 	}
 
 	bool FRigidBodyAsyncGT::IsActive() const
diff --git a/Engine/Plugins/Experimental/ChaosRigidPhysicsAsync/Source/ChaosRigidPhysicsAsync/Private/ChaosRigidPhysicsAsync/RigidBodyContainerAsync.cpp b/Engine/Plugins/Experimental/ChaosRigidPhysicsAsync/Source/ChaosRigidPhysicsAsync/Private/ChaosRigidPhysicsAsync/RigidBodyContainerAsync.cpp
deleted file mode 100644
diff --git a/Engine/Plugins/Experimental/ChaosRigidPhysicsAsync/Source/ChaosRigidPhysicsAsync/Private/ChaosRigidPhysicsAsync/RigidSceneAsync.cpp b/Engine/Plugins/Experimental/ChaosRigidPhysicsAsync/Source/ChaosRigidPhysicsAsync/Private/ChaosRigidPhysicsAsync/RigidSceneAsync.cpp
--- a/Engine/Plugins/Experimental/ChaosRigidPhysicsAsync/Source/ChaosRigidPhysicsAsync/Private/ChaosRigidPhysicsAsync/RigidSceneAsync.cpp
+++ b/Engine/Plugins/Experimental/ChaosRigidPhysicsAsync/Source/ChaosRigidPhysicsAsync/Private/ChaosRigidPhysicsAsync/RigidSceneAsync.cpp
@@ -19,8 +19,7 @@ namespace Chaos::Rigids::Async
 		, SceneSettings(InSceneSettings)
 		, Scene(nullptr)
 		, Owner()
-		, StaticBodyPool()
-		, DynamicBodyPool()
+		, BodyPool(10)
 		, SceneModifiers()
 		, ScenePT()
 		, bOwnsScene(false)
@@ -48,13 +47,7 @@ namespace Chaos::Rigids::Async
 
 		FString StaticName = FString::Format(TEXT("{0} Static"), { Name.Value() });
 
-		StaticBodyPool = MakeUnique<FRigidBodyPoolAsyncGT>(this, FRigidDebugName::Make(TEXT("{0} Statics"), Name.Value()));
-		RegisterBodyContainer(StaticBodyPool.Get());
-
-		DynamicBodyPool = MakeUnique<FRigidBodyPoolAsyncGT>(this, FRigidDebugName::Make(TEXT("{0} Dynamics"), Name.Value()));
-		RegisterBodyContainer(DynamicBodyPool.Get());
-
-		ScenePT = MakeUnique<FRigidSceneAsyncPT>(Name, StaticBodyPool->GetHandle(), DynamicBodyPool->GetHandle());
+		ScenePT = MakeUnique<FRigidSceneAsyncPT>(Name);
 	}
 
 	void FRigidSceneAsyncGT::DestructScene()
@@ -63,13 +56,9 @@ namespace Chaos::Rigids::Async
 
 		ScenePT.Reset();
 
-		UnregisterBodyContainer(DynamicBodyPool.Get());
-		DynamicBodyPool.Reset();
-
-		UnregisterBodyContainer(StaticBodyPool.Get());
-		StaticBodyPool.Reset();
-
 		FRigidScene::DestructScene();
+
+		BodyPool.Reset();
 	}
 
 	void FRigidSceneAsyncGT::Startup()
@@ -103,6 +92,17 @@ namespace Chaos::Rigids::Async
 		bOwnsScene = false;
 	}
 
+	void FRigidSceneAsyncGT::RegisterBody(UE::Physics::IRigidBody* InBody)
+	{
+		const FRigidObjectId BodyId = GetBodyRegistry()->AddObject(InBody);
+		InBody->SetId(BodyId);
+	}
+
+	void FRigidSceneAsyncGT::UnRegisterBody(const UE::Physics::FRigidBodyHandle& InBodyHandle, const UE::Physics::IRigidBody* InBody)
+	{
+		GetBodyRegistry()->RemoveObjectChecked(InBodyHandle.GetId(), InBody);
+	}
+
 	void FRigidSceneAsyncGT::RegisterBodyContainer(UE::Physics::IRigidBodyContainer* InContainer)
 	{
 		using namespace UE::Physics;
@@ -214,36 +214,20 @@ namespace Chaos::Rigids::Async
 		using namespace UE::Physics;
 		CheckGameWriteLocked();
 
-		FRigidBodyAsyncGT* BodyGT = nullptr;
+		FRigidBodyAsyncGT* BodyGT = BodyPool.Alloc(this, InName, InMovementType);
+		UE_RIGIDPHYSICS_CHECK(BodyGT != nullptr);
 
-		if (InMovementType == ERigidMovementType::Static)
+		if (BodyGT != nullptr)
 		{
-			BodyGT = StaticBodyPool->CreateBody(InName, InMovementType);
-			if (BodyGT != nullptr)
-			{
-				if (FPBDRigidsSolver* Solver = GetSolver())
-				{
-					Solver->EnqueueCommandImmediate(
-						[this, BodyHandle = BodyGT->GetHandle(), Proxy = BodyGT->GetProxy()]()
-						{
-							ScenePT->CreateStaticBodyPT(BodyHandle, Proxy);
-						});
-				}
-			}
-		}
-		else
-		{
-			BodyGT = DynamicBodyPool->CreateBody(InName, InMovementType);
-			if (BodyGT != nullptr)
+			RegisterBody(BodyGT);
+
+			if (FPBDRigidsSolver* Solver = GetSolver())
 			{
-				if (FPBDRigidsSolver* Solver = GetSolver())
-				{
-					Solver->EnqueueCommandImmediate(
-						[this, BodyHandle = BodyGT->GetHandle(), Proxy = BodyGT->GetProxy()]()
-						{
-							ScenePT->CreateDynamicBodyPT(BodyHandle, Proxy);
-						});
-				}
+				Solver->EnqueueCommandImmediate(
+					[this, BodyHandle = BodyGT->GetHandle(), Proxy = BodyGT->GetProxy()]()
+					{
+						ScenePT->CreateBodyPT(BodyHandle, Proxy);
+					});
 			}
 		}
 
@@ -255,23 +239,25 @@ namespace Chaos::Rigids::Async
 		using namespace UE::Physics;
 		CheckGameWriteLocked();
 
-		FRigidObjectId BodyPoolId = InBody->GetHandle().GetContainerHandle().GetId();
-		IRigidBodyContainer* BodyPool = (BodyPoolId == StaticBodyPool->GetHandle().GetId()) ? StaticBodyPool.Get() : DynamicBodyPool.Get();
+		UE_RIGIDPHYSICS_CHECK(InBody != nullptr);
 
-		if (BodyPool != nullptr)
+		FRigidBodyAsyncGT* Body = InBody->AsAChecked<FRigidBodyAsyncGT>();
+		UE_RIGIDPHYSICS_CHECK(Body);
+		if (Body != nullptr)
 		{
-			UE_RIGIDPHYSICS_CHECK(BodyPool->GetHandle().GetId() == BodyPoolId);
-
+			const FRigidBodyHandle BodyHandle = Body->GetHandle();
 			if (FPBDRigidsSolver* Solver = GetSolver())
 			{
 				Solver->EnqueueCommandImmediate(
-					[this, BodyHandle = InBody->GetHandle()]()
+					[this, BodyHandle]()
 					{
 						ScenePT->DestroyBodyPT(BodyHandle);
 					});
 			}
 
-			BodyPool->DestroyBody(InBody);
+			Body->Deactivate();
+			UnRegisterBody(BodyHandle, Body);
+			BodyPool.Free(Body);
 		}
 	}
 
@@ -501,30 +487,17 @@ namespace Chaos::Rigids::Async
 {
 	UE_RIGIDPHYSICS_RIGIDTYPED_IMPL(FRigidSceneAsyncPT, UE::Physics::IRigidScene);
 
-	FRigidSceneAsyncPT::FRigidSceneAsyncPT(
-		const UE::Physics::FRigidDebugName& InName,
-		const UE::Physics::FRigidBodyContainerHandle& InStaticBodyPoolHandle,
-		const UE::Physics::FRigidBodyContainerHandle& InDynamicBodyPoolHandle)
+	FRigidSceneAsyncPT::FRigidSceneAsyncPT(const UE::Physics::FRigidDebugName& InName)
 		: Name(InName)
 		, Solver(nullptr)
+		, BodyPool(10)
 		, SceneModifiers()
 		, bIsStarted(false)
 	{
-		StaticBodyPool = MakeUnique<FRigidBodyPoolAsyncPT>(InName, InStaticBodyPoolHandle);
-		RegisterBodyContainer(StaticBodyPool.Get());
-
-		DynamicBodyPool = MakeUnique<FRigidBodyPoolAsyncPT>(InName, InDynamicBodyPoolHandle);
-		RegisterBodyContainer(DynamicBodyPool.Get());
 	}
 
 	FRigidSceneAsyncPT::~FRigidSceneAsyncPT()
 	{
-		UnregisterBodyContainer(DynamicBodyPool.Get());
-		DynamicBodyPool.Reset();
-
-		UnregisterBodyContainer(StaticBodyPool.Get());
-		StaticBodyPool.Reset();
-
 		UE_RIGIDPHYSICS_CHECK(BodyContainerRegistry.IsEmpty());
 	}
 
@@ -652,31 +625,49 @@ namespace Chaos::Rigids::Async
 		CheckSimContext();
 	}
 
-	void FRigidSceneAsyncPT::CreateStaticBodyPT(const UE::Physics::FRigidBodyHandle& InBodyHandle, FSingleParticlePhysicsProxy* InProxy)
+	void FRigidSceneAsyncPT::CreateBodyPT(const UE::Physics::FRigidBodyHandle& InBodyHandle, FSingleParticlePhysicsProxy* InProxy)
 	{
 		CheckSimContext();
+		UE_RIGIDPHYSICS_CHECK(GetBody(InBodyHandle.GetId()) == nullptr);
 
-		StaticBodyPool->CreateBodyPT(InBodyHandle, InProxy);
+		FRigidBodyAsyncPT* Body = BodyPool.Alloc(InBodyHandle, InProxy);
+		RegisterBody(InBodyHandle, Body);
 	}
 
-	void FRigidSceneAsyncPT::CreateDynamicBodyPT(const UE::Physics::FRigidBodyHandle& InBodyHandle, FSingleParticlePhysicsProxy* InProxy)
+	void FRigidSceneAsyncPT::DestroyBodyPT(const UE::Physics::FRigidBodyHandle& InBodyHandle)
 	{
-		CheckSimContext();
+		using UE::Physics::IRigidBodyContainer;
 
-		DynamicBodyPool->CreateBodyPT(InBodyHandle, InProxy);
+		UE::Physics::IRigidBody* Body = GetBody(InBodyHandle.GetId());
+		UnRegisterBody(InBodyHandle, Body);
+		if (FRigidBodyAsyncPT* TypedBody = Body->AsAChecked<FRigidBodyAsyncPT>())
+		{
+			BodyPool.Free(TypedBody);
+		}
+	}
+	
+	void FRigidSceneAsyncPT::RegisterBody(const UE::Physics::FRigidBodyHandle& InBodyHandle, UE::Physics::IRigidBody* InBody)
+	{
+		BodyRegistry.Add(InBodyHandle.GetId(), InBody);
 	}
 
-	void FRigidSceneAsyncPT::DestroyBodyPT(const UE::Physics::FRigidBodyHandle& InBodyHandle)
+	void FRigidSceneAsyncPT::UnRegisterBody(const UE::Physics::FRigidBodyHandle& InBodyHandle, const UE::Physics::IRigidBody* InBody)
 	{
-		using UE::Physics::IRigidBodyContainer;
+		UE_RIGIDPHYSICS_CHECK(GetBody(InBodyHandle.GetId()) == InBody);
+		BodyRegistry.Remove(InBodyHandle.GetId());
+	}
+	
+	UE::Physics::IRigidBody* FRigidSceneAsyncPT::GetBody(const UE::Physics::FRigidObjectId& InBodyId) const
+	{
+		using namespace UE::Physics;
+		CheckSimContext();
 
-		if (IRigidBodyContainer* BodyContainer = GetBodyContainer(InBodyHandle.GetContainerHandle().GetId()))
+		if (IRigidBody*const* BodyValuePtr = BodyRegistry.Find(InBodyId))
 		{
-			if (FRigidBodyPoolAsyncPT* BodyPool = BodyContainer->AsAChecked<FRigidBodyPoolAsyncPT>())
-			{
-				BodyPool->DestroyBodyPT(InBodyHandle);
-			}
+			return *BodyValuePtr;
 		}
+
+		return nullptr;
 	}
 
 	UE::Physics::IRigidBodyContainer* FRigidSceneAsyncPT::GetBodyContainer(const UE::Physics::FRigidObjectId& InBodyContainerId) const
diff --git a/Engine/Plugins/Experimental/ChaosRigidPhysicsAsync/Source/ChaosRigidPhysicsAsync/Public/ChaosRigidPhysicsAsync/RigidBodyAsync.h b/Engine/Plugins/Experimental/ChaosRigidPhysicsAsync/Source/ChaosRigidPhysicsAsync/Public/ChaosRigidPhysicsAsync/RigidBodyAsync.h
--- a/Engine/Plugins/Experimental/ChaosRigidPhysicsAsync/Source/ChaosRigidPhysicsAsync/Public/ChaosRigidPhysicsAsync/RigidBodyAsync.h
+++ b/Engine/Plugins/Experimental/ChaosRigidPhysicsAsync/Source/ChaosRigidPhysicsAsync/Public/ChaosRigidPhysicsAsync/RigidBodyAsync.h
@@ -34,7 +34,7 @@ namespace Chaos::Rigids::Async
 	public:
 		UE_RIGIDPHYSICS_RIGIDTYPED_DECL(CHAOSRIGIDPHYSICSASYNC_API, FRigidBodyAsyncGT);
 
-		UE_INTERNAL CHAOSRIGIDPHYSICSASYNC_API FRigidBodyAsyncGT(FRigidBodyPoolAsyncGT* InContainer, const UE::Physics::FRigidDebugName& InName, UE::Physics::ERigidMovementType InMovementType);
+		UE_INTERNAL CHAOSRIGIDPHYSICSASYNC_API FRigidBodyAsyncGT(FRigidSceneAsyncGT* InScene, const UE::Physics::FRigidDebugName& InName, UE::Physics::ERigidMovementType InMovementType);
 
 		virtual ~FRigidBodyAsyncGT();
 
@@ -42,6 +42,8 @@ namespace Chaos::Rigids::Async
 		// Begin IRigidBody
 		//
 		UE_INTERNAL virtual UE::Physics::FRigidBodyHandle GetHandle() const override final;
+		UE_INTERNAL CHAOSRIGIDPHYSICSASYNC_API virtual void SetId(const UE::Physics::FRigidObjectId& InId) override final;
+
 		UE_INTERNAL virtual bool IsActive() const override final;
 		UE_INTERNAL virtual void Activate() override final;
 		UE_INTERNAL virtual void Deactivate() override final;
@@ -80,7 +82,6 @@ namespace Chaos::Rigids::Async
 		//
 
 		UE_INTERNAL CHAOSRIGIDPHYSICSASYNC_API bool IsValid() const;
-		UE_INTERNAL CHAOSRIGIDPHYSICSASYNC_API void SetId(const UE::Physics::FRigidObjectId& InId);
 		UE_INTERNAL CHAOSRIGIDPHYSICSASYNC_API const FSingleParticlePhysicsProxy* GetProxy() const;
 		UE_INTERNAL CHAOSRIGIDPHYSICSASYNC_API FSingleParticlePhysicsProxy* GetProxy();
 		UE_INTERNAL CHAOSRIGIDPHYSICSASYNC_API const FRigidBodyHandle_External& GetProxyAPI() const;
@@ -90,7 +91,7 @@ namespace Chaos::Rigids::Async
 
 	private:
 		UE::Physics::FRigidObjectId BodyId;
-		FRigidBodyPoolAsyncGT* BodyContainer;
+		FRigidSceneAsyncGT* SceneAsync = nullptr;
 		FSingleParticlePhysicsProxy* Proxy;
 	};
 
diff --git a/Engine/Plugins/Experimental/ChaosRigidPhysicsAsync/Source/ChaosRigidPhysicsAsync/Public/ChaosRigidPhysicsAsync/RigidBodyContainerAsync.h b/Engine/Plugins/Experimental/ChaosRigidPhysicsAsync/Source/ChaosRigidPhysicsAsync/Public/ChaosRigidPhysicsAsync/RigidBodyContainerAsync.h
deleted file mode 100644
diff --git a/Engine/Plugins/Experimental/ChaosRigidPhysicsAsync/Source/ChaosRigidPhysicsAsync/Public/ChaosRigidPhysicsAsync/RigidSceneAsync.h b/Engine/Plugins/Experimental/ChaosRigidPhysicsAsync/Source/ChaosRigidPhysicsAsync/Public/ChaosRigidPhysicsAsync/RigidSceneAsync.h
--- a/Engine/Plugins/Experimental/ChaosRigidPhysicsAsync/Source/ChaosRigidPhysicsAsync/Public/ChaosRigidPhysicsAsync/RigidSceneAsync.h
+++ b/Engine/Plugins/Experimental/ChaosRigidPhysicsAsync/Source/ChaosRigidPhysicsAsync/Public/ChaosRigidPhysicsAsync/RigidSceneAsync.h
@@ -8,7 +8,6 @@
 #include "Chaos/ObjectPool.h"
 #include "ChaosRigidPhysicsAsync/JointConstraint6DOFAsync.h"
 #include "ChaosRigidPhysicsAsync/RigidBodyAsync.h"
-#include "ChaosRigidPhysicsAsync/RigidBodyContainerAsync.h"
 #include "ChaosRigidPhysicsAsync/RigidFactoryAsync.h"
 #include "ChaosRigidPhysicsAsync/RigidGeometryCollectionAsync.h"
 #include "ChaosRigidPhysicsAsync/RigidSceneSettingsAsync.h"
@@ -89,6 +88,8 @@ namespace Chaos::Rigids::Async
 	public:
 		UE_INTERNAL CHAOSRIGIDPHYSICSASYNC_API void Startup();
 		UE_INTERNAL CHAOSRIGIDPHYSICSASYNC_API void Shutdown();
+		UE_INTERNAL CHAOSRIGIDPHYSICSASYNC_API void RegisterBody(UE::Physics::IRigidBody* InBody);
+		UE_INTERNAL CHAOSRIGIDPHYSICSASYNC_API void UnRegisterBody(const UE::Physics::FRigidBodyHandle& InBodyHandle, const UE::Physics::IRigidBody* InBody);
 		UE_INTERNAL CHAOSRIGIDPHYSICSASYNC_API void RegisterBodyContainer(UE::Physics::IRigidBodyContainer* InContainer);
 		UE_INTERNAL CHAOSRIGIDPHYSICSASYNC_API void UnregisterBodyContainer(UE::Physics::IRigidBodyContainer* InContainer);
 		UE_INTERNAL CHAOSRIGIDPHYSICSASYNC_API void ActivateBody(FRigidBodyAsyncGT* InBody);
@@ -112,8 +113,7 @@ namespace Chaos::Rigids::Async
 		FChaosScene* Scene;
 		TStrongObjectPtr<UObject> Owner;
 
-		TUniquePtr<FRigidBodyPoolAsyncGT> StaticBodyPool;
-		TUniquePtr<FRigidBodyPoolAsyncGT> DynamicBodyPool;
+		TObjectPool<FRigidBodyAsyncGT> BodyPool;
 
 		TArray<UE::Physics::IRigidSceneModifier*> SceneModifiers;
 
@@ -129,10 +129,7 @@ namespace Chaos::Rigids::Async
 	public:
 		UE_RIGIDPHYSICS_RIGIDTYPED_DECL(CHAOSRIGIDPHYSICSASYNC_API, FRigidSceneAsyncPT);
 
-		UE_INTERNAL CHAOSRIGIDPHYSICSASYNC_API FRigidSceneAsyncPT(
-			const UE::Physics::FRigidDebugName& InName, 
-			const UE::Physics::FRigidBodyContainerHandle& InStaticBodyPoolHandle,
-			const UE::Physics::FRigidBodyContainerHandle& InDynamicBodyPoolHandle);
+		UE_INTERNAL CHAOSRIGIDPHYSICSASYNC_API FRigidSceneAsyncPT(const UE::Physics::FRigidDebugName& InName);
 
 		CHAOSRIGIDPHYSICSASYNC_API virtual ~FRigidSceneAsyncPT();
 
@@ -151,6 +148,7 @@ namespace Chaos::Rigids::Async
 		//
 	public:
 		UE_INTERNAL virtual UE::Physics::FRigidSceneHandle GetHandle() const override final;
+		UE_INTERNAL virtual UE::Physics::IRigidBody* GetBody(const UE::Physics::FRigidObjectId& InBodyId) const override final;
 		UE_INTERNAL virtual UE::Physics::IRigidBodyContainer* GetBodyContainer(const UE::Physics::FRigidObjectId& InBodyContainerId) const override final;
 		UE_INTERNAL virtual FRigidGeometryCollectionAsyncPT* GetGeometryCollection(const UE::Physics::FRigidObjectId& InGeometryCollectionId) const override final;
 		UE_INTERNAL virtual void RegisterModifier(UE::Physics::IRigidSceneModifier* Modifier) override final;
@@ -169,9 +167,10 @@ namespace Chaos::Rigids::Async
 		UE_INTERNAL CHAOSRIGIDPHYSICSASYNC_API void UnregisterBodyContainer(UE::Physics::IRigidBodyContainer* InContainer);
 		UE_INTERNAL CHAOSRIGIDPHYSICSASYNC_API void CreateGeometryCollectionPT(const UE::Physics::FRigidBodyContainerHandle& InHandle, FGeometryCollectionPhysicsProxy* InProxy);
 		UE_INTERNAL CHAOSRIGIDPHYSICSASYNC_API void DestroyGeometryCollectionPT(const UE::Physics::FRigidBodyContainerHandle& InHandle);
-		UE_INTERNAL CHAOSRIGIDPHYSICSASYNC_API void CreateStaticBodyPT(const UE::Physics::FRigidBodyHandle& InBodyHandle, FSingleParticlePhysicsProxy* InProxy);
-		UE_INTERNAL CHAOSRIGIDPHYSICSASYNC_API void CreateDynamicBodyPT(const UE::Physics::FRigidBodyHandle& InBodyHandle, FSingleParticlePhysicsProxy* InProxy);
+		UE_INTERNAL CHAOSRIGIDPHYSICSASYNC_API void CreateBodyPT(const UE::Physics::FRigidBodyHandle& InBodyHandle, FSingleParticlePhysicsProxy* InProxy);
 		UE_INTERNAL CHAOSRIGIDPHYSICSASYNC_API void DestroyBodyPT(const UE::Physics::FRigidBodyHandle& InBodyHandle);
+		UE_INTERNAL CHAOSRIGIDPHYSICSASYNC_API void RegisterBody(const UE::Physics::FRigidBodyHandle& InBodyHandle, UE::Physics::IRigidBody* InBody);
+		UE_INTERNAL CHAOSRIGIDPHYSICSASYNC_API void UnRegisterBody(const UE::Physics::FRigidBodyHandle& InBodyHandle, const UE::Physics::IRigidBody* InBody);
 		UE_INTERNAL CHAOSRIGIDPHYSICSASYNC_API void ActivateBody(const UE::Physics::FRigidBodyHandle& InBodyHandle);
 		UE_INTERNAL CHAOSRIGIDPHYSICSASYNC_API void DeactivateBody(const UE::Physics::FRigidBodyHandle& InBodyHandle);
 		
@@ -192,9 +191,9 @@ namespace Chaos::Rigids::Async
 		UE::Physics::FRigidDebugName Name;
 		UE::Physics::FRigidSceneHandle SceneHandle;
 		FPBDRigidsSolver* Solver;
-		TUniquePtr<FRigidBodyPoolAsyncPT> StaticBodyPool;
-		TUniquePtr<FRigidBodyPoolAsyncPT> DynamicBodyPool;
+		TObjectPool<FRigidBodyAsyncPT> BodyPool;
 		TMap<UE::Physics::FRigidObjectId, UE::Physics::IRigidBodyContainer*> BodyContainerRegistry;
+		TMap<UE::Physics::FRigidObjectId, UE::Physics::IRigidBody*> BodyRegistry;
 		TMap<UE::Physics::FRigidObjectId, UE::Physics::IJointConstraint*> JointConstraintRegistry;
 		TArray<UE::Physics::IRigidSceneModifier*> SceneModifiers;
 		bool bIsStarted;
diff --git a/Engine/Plugins/Experimental/ChaosRigidPhysicsAsync/Tests/Private/GeometryCollectionTests.cpp b/Engine/Plugins/Experimental/ChaosRigidPhysicsAsync/Tests/Private/GeometryCollectionTests.cpp
--- a/Engine/Plugins/Experimental/ChaosRigidPhysicsAsync/Tests/Private/GeometryCollectionTests.cpp
+++ b/Engine/Plugins/Experimental/ChaosRigidPhysicsAsync/Tests/Private/GeometryCollectionTests.cpp
@@ -3,11 +3,11 @@
 #include "Chaos/LowLevelTest/ChaosTestHarness.h"
 #include "Chaos/LowLevelTest/ChaosTestScene.h"
 #include "ChaosRigidPhysicsAsyncTest.h"
-#include "ChaosRigidPhysicsAsync/RigidBodyContainerAsync.h"
 #include "ChaosRigidPhysicsAsync/RigidFactoryAsync.h"
 #include "ChaosRigidPhysicsAsync/RigidGeometryCollectionAsync.h"
 #include "ChaosRigidPhysicsAsync/RigidSceneAsync.h"
 #include "ChaosRigidPhysicsAsync/RigidSceneSettingsAsync.h"
+#include "RigidPhysics/RigidBodyContainer.h"
 #include "RigidPhysics/RigidContext.h"
 
 #if UE_RIGIDPHYSICS_API_ENABLED
@@ -44,25 +44,6 @@ namespace Chaos::LowLevelTest
 			CHECK_FALSE(GCPtr.IsValid());
 		}
 
-		// We can safely (fail to) cast a valid body container to a GC, giving a null GC ptr
-		if (FRigidContextGameRW Context = SceneHandle.LockRW())
-		{
-			// Create a standalone body (i.e., not in a GC)
-			TRigidBodyPtr<FRigidContextGameRW> BodyPtr = Context->CreateBody(FRigidDebugName::Make(TEXT("Body1")), ERigidMovementType::Static);
-			REQUIRE(BodyPtr.IsValid());
-
-			// Get the container that owns the body - it should be a body pool, not a GC
-			TRigidBodyContainerPtr<FRigidContextGameRW> BodyContainerPtr = Context->GetBodyContainer(BodyPtr->GetContainerHandle().GetId());
-			CHECK(BodyContainerPtr.IsValid());
-			CHECK(BodyContainerPtr->GetTypeId() == FRigidBodyPoolAsyncGT::GetStaticTypeId());
-
-			// Attempt to cast the container to a GC - should fail gracefully
-			TRigidGeometryCollectionPtr<FRigidContextGameRW> GCPtr = BodyContainerPtr;
-			CHECK_FALSE(GCPtr.IsValid());
-
-			Context->DestroyBody(MoveTemp(BodyPtr));
-		}
-
 		// We can cast a valid GCPtr to a valid BodyContainerPtr
 		if (FRigidContextGameRW Context = SceneHandle.LockRW())
 		{
diff --git a/Engine/Plugins/Experimental/ChaosRigidPhysicsAsync/Tests/Private/RigidTestFixture.cpp b/Engine/Plugins/Experimental/ChaosRigidPhysicsAsync/Tests/Private/RigidTestFixture.cpp
--- a/Engine/Plugins/Experimental/ChaosRigidPhysicsAsync/Tests/Private/RigidTestFixture.cpp
+++ b/Engine/Plugins/Experimental/ChaosRigidPhysicsAsync/Tests/Private/RigidTestFixture.cpp
@@ -8,8 +8,8 @@
 
 #include "ChaosRigidPhysicsAsync/RigidFactoryAsync.h"
 #include "ChaosRigidPhysicsAsync/RigidBodyAsync.h"
+#include "ChaosRigidPhysicsAsync/RigidGeometryCollectionAsync.h"
 #include "ChaosRigidPhysicsAsync/RigidSceneAsync.h"
-#include "ChaosRigidPhysicsAsync/RigidSceneSettingsAsync.h"
 
 #include "RigidPhysics/RigidBody.h"
 #include "RigidPhysics/RigidCollider.h"
diff --git a/Engine/Source/Runtime/Experimental/RigidPhysics/Private/RigidPhysics/Internal/IRigidScene.cpp b/Engine/Source/Runtime/Experimental/RigidPhysics/Private/RigidPhysics/Internal/IRigidScene.cpp
--- a/Engine/Source/Runtime/Experimental/RigidPhysics/Private/RigidPhysics/Internal/IRigidScene.cpp
+++ b/Engine/Source/Runtime/Experimental/RigidPhysics/Private/RigidPhysics/Internal/IRigidScene.cpp
@@ -15,9 +15,15 @@ namespace UE::Physics
 	{
 		// TODO_CHAOSAPI: registry slack should be a parameter
 		BodyContainerRegistry = MakeUnique<FRigidBodyContainerRegistry>(10);
+		BodyRegistry = MakeUnique<FRigidBodyRegistry>(10);
 		JointConstraintRegistry = MakeUnique<FJointConstraintRegistry>(10);
 	}
 
+	IRigidBody* FRigidScene::GetBody(const FRigidObjectId& InBodyId) const
+	{
+		return BodyRegistry->FindObject(InBodyId);
+	}
+
 	IRigidBodyContainer* FRigidScene::GetBodyContainer(const FRigidObjectId& InBodyContainerId) const
 	{
 		return BodyContainerRegistry->FindObject(InBodyContainerId);
@@ -35,6 +41,7 @@ namespace UE::Physics
 	void FRigidScene::DestructScene()
 	{
 		BodyContainerRegistry.Reset();
+		BodyRegistry.Reset();
 		JointConstraintRegistry.Reset();
 	}
 
@@ -43,6 +50,11 @@ namespace UE::Physics
 		return BodyContainerRegistry.Get();
 	}
 
+	FRigidBodyRegistry* FRigidScene::GetBodyRegistry() const
+	{
+		return BodyRegistry.Get();
+	}
+
 	FJointConstraintRegistry* FRigidScene::GetJointConstraintRegistry() const
 	{
 		return JointConstraintRegistry.Get();
diff --git a/Engine/Source/Runtime/Experimental/RigidPhysics/Private/RigidPhysics/RigidBodyHandle.cpp b/Engine/Source/Runtime/Experimental/RigidPhysics/Private/RigidPhysics/RigidBodyHandle.cpp
--- a/Engine/Source/Runtime/Experimental/RigidPhysics/Private/RigidPhysics/RigidBodyHandle.cpp
+++ b/Engine/Source/Runtime/Experimental/RigidPhysics/Private/RigidPhysics/RigidBodyHandle.cpp
@@ -1,15 +1,15 @@
 // Copyright Epic Games, Inc. All Rights Reserved.
 #include "RigidPhysics/RigidBodyHandle.h"
 
-#include "RigidPhysics/RigidBodyContainer.h"
-
 #if UE_RIGIDPHYSICS_API_ENABLED
 
+#include "RigidPhysics/RigidScene.h"
+
 namespace UE::Physics
 {
-	FRigidBodyHandle::FRigidBodyHandle(const FRigidBodyContainerHandle& InBodyContainerHandle, const FRigidObjectId& InBodyId)
+	FRigidBodyHandle::FRigidBodyHandle(const FRigidSceneHandle& InSceneHandle, const FRigidObjectId& InBodyId)
 		: BodyId(InBodyId)
-		, BodyContainerHandle(InBodyContainerHandle)
+		, SceneHandle(InSceneHandle)
 	{
 	}
 
@@ -18,21 +18,16 @@ namespace UE::Physics
 		return BodyId;
 	}
 
-	const FRigidBodyContainerHandle& FRigidBodyHandle::GetContainerHandle() const
-	{
-		return BodyContainerHandle;
-	}
-
 	void FRigidBodyHandle::Reset()
 	{
-		BodyContainerHandle.Reset();
+		SceneHandle.Reset();
 		BodyId.Reset();
 	}
 
-	IRigidBody* FRigidBodyHandle::PinInternal(const IRigidBodyContainer* Container) const
+	IRigidBody* FRigidBodyHandle::PinInternal(const IRigidScene* Scene) const
 	{
-		return Container->GetBody(GetId());
+		return Scene->GetBody(GetId());
 	}
 } // namespace UE::Physics
 
-#endif // UE_RIGIDPHYSICS_API_ENABLED
\ No newline at end of file
+#endif // UE_RIGIDPHYSICS_API_ENABLED
diff --git a/Engine/Source/Runtime/Experimental/RigidPhysics/Public/RigidPhysics/Internal/IRigidBody.h b/Engine/Source/Runtime/Experimental/RigidPhysics/Public/RigidPhysics/Internal/IRigidBody.h
--- a/Engine/Source/Runtime/Experimental/RigidPhysics/Public/RigidPhysics/Internal/IRigidBody.h
+++ b/Engine/Source/Runtime/Experimental/RigidPhysics/Public/RigidPhysics/Internal/IRigidBody.h
@@ -49,6 +49,11 @@ namespace UE::Physics
 		// The handle to this body
 		UE_INTERNAL RIGIDPHYSICS_API virtual FRigidBodyHandle GetHandle() const = 0;
 
+		UE_INTERNAL RIGIDPHYSICS_API virtual void SetId(const FRigidObjectId& InId)
+		{
+			RIGIDPHYSICS_API_UNSUPPORTED();
+		}
+
 		// Is this body Static? Static bodies can be moved but have no velocity.
 		UE_INTERNAL RIGIDPHYSICS_API virtual bool IsStatic() const = 0;
 
diff --git a/Engine/Source/Runtime/Experimental/RigidPhysics/Public/RigidPhysics/Internal/IRigidScene.h b/Engine/Source/Runtime/Experimental/RigidPhysics/Public/RigidPhysics/Internal/IRigidScene.h
--- a/Engine/Source/Runtime/Experimental/RigidPhysics/Public/RigidPhysics/Internal/IRigidScene.h
+++ b/Engine/Source/Runtime/Experimental/RigidPhysics/Public/RigidPhysics/Internal/IRigidScene.h
@@ -25,6 +25,8 @@ namespace UE::Physics
 
 		UE_INTERNAL RIGIDPHYSICS_API virtual FRigidSceneHandle GetHandle() const = 0;
 
+		UE_INTERNAL RIGIDPHYSICS_API virtual IRigidBody* GetBody(const FRigidObjectId& InBodyId) const = 0;
+
 		UE_INTERNAL RIGIDPHYSICS_API virtual IRigidBodyContainer* GetBodyContainer(const FRigidObjectId& InBodyContainerId) const = 0;
 
 		UE_INTERNAL RIGIDPHYSICS_API virtual IRigidGeometryCollection* GetGeometryCollection(const FRigidObjectId& InGeometryCollectionId) const = 0;
@@ -102,6 +104,7 @@ namespace UE::Physics
 		RIGIDPHYSICS_API FRigidScene();
 		virtual ~FRigidScene() = default;
 
+		UE_INTERNAL RIGIDPHYSICS_API virtual IRigidBody* GetBody(const FRigidObjectId& InBodyId) const override final;
 		UE_INTERNAL RIGIDPHYSICS_API virtual IRigidBodyContainer* GetBodyContainer(const FRigidObjectId& InBodyContainerId) const override final;
 		UE_INTERNAL RIGIDPHYSICS_API virtual IRigidGeometryCollection* GetGeometryCollection(const FRigidObjectId& InGeometryCollectionId) const override final;
 
@@ -109,10 +112,12 @@ namespace UE::Physics
 		UE_INTERNAL RIGIDPHYSICS_API void DestructScene();
 
 		UE_INTERNAL RIGIDPHYSICS_API FRigidBodyContainerRegistry* GetBodyContainerRegistry() const;
+		UE_INTERNAL RIGIDPHYSICS_API FRigidBodyRegistry* GetBodyRegistry() const;
 		UE_INTERNAL RIGIDPHYSICS_API FJointConstraintRegistry* GetJointConstraintRegistry() const;
 
 	private:
 		TUniquePtr<FRigidBodyContainerRegistry> BodyContainerRegistry;
+		TUniquePtr<FRigidBodyRegistry> BodyRegistry;
 		TUniquePtr<FJointConstraintRegistry> JointConstraintRegistry;
 	};
 } // namespace UE::Physics
diff --git a/Engine/Source/Runtime/Experimental/RigidPhysics/Public/RigidPhysics/RigidBody.h b/Engine/Source/Runtime/Experimental/RigidPhysics/Public/RigidPhysics/RigidBody.h
--- a/Engine/Source/Runtime/Experimental/RigidPhysics/Public/RigidPhysics/RigidBody.h
+++ b/Engine/Source/Runtime/Experimental/RigidPhysics/Public/RigidPhysics/RigidBody.h
@@ -8,7 +8,6 @@
 #include "Math/Bounds.h"
 #include "Math/Transform.h"
 #include "RigidPhysics/Internal/IRigidBody.h"
-#include "RigidPhysics/RigidBodyContainer.h"
 #include "RigidPhysics/RigidBodyHandle.h"
 #include "RigidPhysics/RigidCollider.h"
 #include "RigidPhysics/RigidObjectId.h"
@@ -74,11 +73,6 @@ namespace UE::Physics
 			return FRigidBodyHandle();
 		}
 
-		FRigidBodyContainerHandle GetContainerHandle() const
-		{
-			return GetHandle().GetContainerHandle();
-		}
-
 		bool IsActive() const
 		{
 			return BodyRaw->IsActive();
diff --git a/Engine/Source/Runtime/Experimental/RigidPhysics/Public/RigidPhysics/RigidBodyHandle.h b/Engine/Source/Runtime/Experimental/RigidPhysics/Public/RigidPhysics/RigidBodyHandle.h
--- a/Engine/Source/Runtime/Experimental/RigidPhysics/Public/RigidPhysics/RigidBodyHandle.h
+++ b/Engine/Source/Runtime/Experimental/RigidPhysics/Public/RigidPhysics/RigidBodyHandle.h
@@ -5,9 +5,9 @@
 
 #if UE_RIGIDPHYSICS_API_ENABLED
 
-#include "RigidPhysics/RigidBodyContainerHandle.h"
 #include "RigidPhysics/RigidObjectId.h"
 #include "RigidPhysics/RigidObjectPtr.h"
+#include "RigidPhysics/RigidSceneHandle.h"
 
 namespace UE::Physics
 {
@@ -15,22 +15,15 @@ namespace UE::Physics
 	template <typename ContextType>
 	class UE_INTERNAL TRigidBodyPtrImpl;
 
-	// A handle to a rigid body in a container.
-	// This handle may be copied and stored, and can be used to retrieve a
-	// body reference by "pinning" it and providing the context.
-	// 
-	// Example contexts are: 
-	// - Read-Only Game-Thread
-	// - Read-Write Game Thread
-	// - Read-Write Physics-Thread Callback
-	//
+	// A handle to a rigid body that may be copied and stored.
+	// A body reference can be retrieved by pinning the handle with a context.
 	class FRigidBodyHandle
 	{
 	public:
 		FRigidBodyHandle() = default;
 		FRigidBodyHandle(const FRigidBodyHandle&) = default;
 
-		UE_INTERNAL RIGIDPHYSICS_API FRigidBodyHandle(const FRigidBodyContainerHandle& InBodyContainerHandle, const FRigidObjectId& InBodyId);
+		UE_INTERNAL RIGIDPHYSICS_API FRigidBodyHandle(const FRigidSceneHandle& InSceneHandle, const FRigidObjectId& InBodyId);
 
 		friend bool operator==(const FRigidBodyHandle&, const FRigidBodyHandle&) = default;
 		friend bool operator!=(const FRigidBodyHandle&, const FRigidBodyHandle&) = default;
@@ -41,24 +34,21 @@ namespace UE::Physics
 		{
 			if (InContext.IsValid())
 			{
-				if (TRigidBodyContainerPtr<ContextType> BodyContainer = BodyContainerHandle.Pin(InContext))
-				{
-					return PinInternal(BodyContainer.Get());
-				}
+				const IRigidScene* Scene = InContext.GetScene().Get();
+				return PinInternal(Scene);
 			}
 			return TRigidBodyPtr<ContextType>();
 		}
 
 		UE_INTERNAL RIGIDPHYSICS_API const FRigidObjectId& GetId() const;
-		UE_INTERNAL RIGIDPHYSICS_API const FRigidBodyContainerHandle& GetContainerHandle() const;
 		UE_INTERNAL RIGIDPHYSICS_API void Reset();
 
 	private:
-		RIGIDPHYSICS_API IRigidBody* PinInternal(const IRigidBodyContainer* Container) const;
+		RIGIDPHYSICS_API IRigidBody* PinInternal(const IRigidScene* Scene) const;
 
 		FRigidObjectId BodyId;
-		FRigidBodyContainerHandle BodyContainerHandle;
+		FRigidSceneHandle SceneHandle;
 	};
 } // namespace UE::Physics
 
-#endif // UE_RIGIDPHYSICS_API_ENABLED
\ No newline at end of file
+#endif // UE_RIGIDPHYSICS_API_ENABLED
diff --git a/Engine/Source/Runtime/Experimental/RigidPhysics/Public/RigidPhysics/RigidGeometryCollection.h b/Engine/Source/Runtime/Experimental/RigidPhysics/Public/RigidPhysics/RigidGeometryCollection.h
--- a/Engine/Source/Runtime/Experimental/RigidPhysics/Public/RigidPhysics/RigidGeometryCollection.h
+++ b/Engine/Source/Runtime/Experimental/RigidPhysics/Public/RigidPhysics/RigidGeometryCollection.h
@@ -6,6 +6,7 @@
 #if UE_RIGIDPHYSICS_API_ENABLED
 
 #include "RigidPhysics/Internal/IRigidGeometryCollection.h"
+#include "RigidPhysics/RigidBodyContainer.h"
 #include "RigidPhysics/RigidGeometryCollectionHandle.h"
 #include "RigidPhysics/RigidObjectPtr.h"
 
diff --git a/Engine/Source/Runtime/Experimental/RigidPhysics/Public/RigidPhysics/RigidObjectRegistry.h b/Engine/Source/Runtime/Experimental/RigidPhysics/Public/RigidPhysics/RigidObjectRegistry.h
--- a/Engine/Source/Runtime/Experimental/RigidPhysics/Public/RigidPhysics/RigidObjectRegistry.h
+++ b/Engine/Source/Runtime/Experimental/RigidPhysics/Public/RigidPhysics/RigidObjectRegistry.h
@@ -165,6 +165,12 @@ namespace UE::Physics
 			// the free indices.
 		}
 
+		UE_INTERNAL void RemoveObjectChecked(const FRigidObjectId& InObjectId, const ObjectType* Object)
+		{
+			UE_RIGIDPHYSICS_CHECK(FindObject(InObjectId) == Object);
+			RemoveObject(InObjectId);
+		}
+
 		UE_INTERNAL FObjectType* FindObject(const FRigidObjectId& InId) const
 		{
 			if (ObjectsById.IsValidIndex(InId.GetArrayIndex()))
