Fix linking issue on Engine in ChaosRigidPhysicsAsync after shape changes.

#Synced-CL 51630736

diff --git a/Engine/Plugins/Experimental/ChaosRigidPhysicsAsync/Source/ChaosRigidPhysicsAsync/ChaosRigidPhysicsAsync.Build.cs b/Engine/Plugins/Experimental/ChaosRigidPhysicsAsync/Source/ChaosRigidPhysicsAsync/ChaosRigidPhysicsAsync.Build.cs
--- a/Engine/Plugins/Experimental/ChaosRigidPhysicsAsync/Source/ChaosRigidPhysicsAsync/ChaosRigidPhysicsAsync.Build.cs
+++ b/Engine/Plugins/Experimental/ChaosRigidPhysicsAsync/Source/ChaosRigidPhysicsAsync/ChaosRigidPhysicsAsync.Build.cs
@@ -13,6 +13,7 @@ namespace UnrealBuildTool.Rules
                 new string[] {
                 "Core",
                 "CoreUObject",
+				"Engine",
                 "RigidPhysics",
 				"DataflowCore",
 				"DataflowEngine",
