Chaos API add methods for skeletal mesh component body setup:

Add setters for additional joint properties. Add COM and ROM setting for bodies. Add method to set disabled collision pairs in scene.

#Synced-CL 51773069

diff --git a/Engine/Plugins/Experimental/ChaosRigidPhysicsAsync/Source/ChaosRigidPhysicsAsync/Private/ChaosRigidPhysicsAsync/JointConstraint6DOFAsync.cpp b/Engine/Plugins/Experimental/ChaosRigidPhysicsAsync/Source/ChaosRigidPhysicsAsync/Private/ChaosRigidPhysicsAsync/JointConstraint6DOFAsync.cpp
--- a/Engine/Plugins/Experimental/ChaosRigidPhysicsAsync/Source/ChaosRigidPhysicsAsync/Private/ChaosRigidPhysicsAsync/JointConstraint6DOFAsync.cpp
+++ b/Engine/Plugins/Experimental/ChaosRigidPhysicsAsync/Source/ChaosRigidPhysicsAsync/Private/ChaosRigidPhysicsAsync/JointConstraint6DOFAsync.cpp
@@ -5,6 +5,7 @@
 
 #include "ChaosRigidPhysicsAsync/RigidSceneAsync.h"
 #include "Chaos/PBDJointConstraintData.h"
+#include "Chaos/ChaosConstraintSettings.h"
 
 namespace Chaos::Rigids::Async
 {
@@ -114,6 +115,147 @@ namespace Chaos::Rigids::Async
 		Constraint->SetLinearMotionTypesZ(Convert(InMotionType));
 	}
 
+	void FJointConstraint6DOFAsyncGT::SetAngularMotionTypesX(const EJointMotionType InMotionType)
+	{
+		Constraint->SetAngularMotionTypesX(Convert(InMotionType));
+	}
+
+	void FJointConstraint6DOFAsyncGT::SetAngularMotionTypesY(const EJointMotionType InMotionType)
+	{
+		Constraint->SetAngularMotionTypesY(Convert(InMotionType));
+	}
+
+	void FJointConstraint6DOFAsyncGT::SetAngularMotionTypesZ(const EJointMotionType InMotionType)
+	{
+		Constraint->SetAngularMotionTypesZ(Convert(InMotionType));
+	}
+
+	void FJointConstraint6DOFAsyncGT::SetAngularLimits(const FVector3d& InAngularLimits)
+	{
+		Constraint->SetAngularLimits(InAngularLimits);
+	}
+
+	void FJointConstraint6DOFAsyncGT::SetCollisionEnabled(const bool bCollisionEnabled)
+	{
+		Constraint->SetCollisionEnabled(bCollisionEnabled);
+	}
+
+	void FJointConstraint6DOFAsyncGT::SetProjectionEnabled(const bool bProjectionEnabled)
+	{
+		Constraint->SetProjectionEnabled(bProjectionEnabled);
+	}
+
+	void FJointConstraint6DOFAsyncGT::SetMassConditioningEnabled(const bool bMassConditioningEnabled)
+	{
+		Constraint->SetMassConditioningEnabled(bMassConditioningEnabled);
+	}
+
+	void FJointConstraint6DOFAsyncGT::SetUseLinearSolver(const bool bUseLinearSolver)
+	{
+		Constraint->SetUseLinearSolver(bUseLinearSolver);
+	}
+
+	void FJointConstraint6DOFAsyncGT::SetSoftLinearLimitsEnabled(const bool bSoftLinearLimitsEnabled)
+	{
+		Constraint->SetSoftLinearLimitsEnabled(bSoftLinearLimitsEnabled);
+	}
+
+	void FJointConstraint6DOFAsyncGT::SetSoftTwistLimitsEnabled(const bool bSoftTwistLimitsEnabled)
+	{
+		Constraint->SetSoftTwistLimitsEnabled(bSoftTwistLimitsEnabled);
+	}
+
+	void FJointConstraint6DOFAsyncGT::SetSoftSwingLimitsEnabled(const bool bSoftSwingLimitsEnabled)
+	{
+		Constraint->SetSoftSwingLimitsEnabled(bSoftSwingLimitsEnabled);
+	}
+
+	void FJointConstraint6DOFAsyncGT::SetSoftLinearStiffness(const double InSoftLinearStiffness)
+	{
+		Constraint->SetSoftLinearStiffness(InSoftLinearStiffness);
+	}
+
+	void FJointConstraint6DOFAsyncGT::SetSoftLinearDamping(const double InSoftLinearDamping)
+	{
+		Constraint->SetSoftLinearDamping(InSoftLinearDamping);
+	}
+
+	void FJointConstraint6DOFAsyncGT::SetSoftTwistStiffness(const double InSoftTwistStiffness)
+	{
+		Constraint->SetSoftTwistStiffness(InSoftTwistStiffness);
+	}
+
+	void FJointConstraint6DOFAsyncGT::SetSoftTwistDamping(const double InSoftTwistDamping)
+	{
+		Constraint->SetSoftTwistDamping(InSoftTwistDamping);
+	}
+
+	void FJointConstraint6DOFAsyncGT::SetSoftSwingStiffness(const double InSoftSwingStiffness)
+	{
+		Constraint->SetSoftSwingStiffness(InSoftSwingStiffness);
+	}
+
+	void FJointConstraint6DOFAsyncGT::SetSoftSwingDamping(const double InSoftSwingDamping)
+	{
+		Constraint->SetSoftSwingDamping(InSoftSwingDamping);
+	}
+
+	void FJointConstraint6DOFAsyncGT::SetProjectionLinearAlpha(const double InProjectionLinearAlpha)
+	{
+		Constraint->SetProjectionLinearAlpha(InProjectionLinearAlpha);
+	}
+
+	void FJointConstraint6DOFAsyncGT::SetProjectionAngularAlpha(const double InProjectionAngularAlpha)
+	{
+		Constraint->SetProjectionAngularAlpha(InProjectionAngularAlpha);
+	}
+
+	void FJointConstraint6DOFAsyncGT::SetProjectionLinearTolerance(const double InProjectionLinearTolerance)
+	{
+		Constraint->SetProjectionLinearTolerance(InProjectionLinearTolerance);
+	}
+
+	void FJointConstraint6DOFAsyncGT::SetProjectionAngularTolerance(const double InProjectionAngularTolerance)
+	{
+		Constraint->SetProjectionAngularTolerance(InProjectionAngularTolerance);
+	}
+
+	void FJointConstraint6DOFAsyncGT::SetLinearContactDistance(const double InLinearContactDistance)
+	{
+		Constraint->SetLinearContactDistance(InLinearContactDistance);
+	}
+
+	void FJointConstraint6DOFAsyncGT::SetTwistContactDistance(const double InTwistContactDistance)
+	{
+		Constraint->SetTwistContactDistance(InTwistContactDistance);
+	}
+
+	void FJointConstraint6DOFAsyncGT::SetSwingContactDistance(const double InSwingContactDistance)
+	{
+		Constraint->SetSwingContactDistance(InSwingContactDistance);
+	}
+
+	void FJointConstraint6DOFAsyncGT::SetAngularDriveStiffness(const FVector3d& InAngularDriveStiffness)
+	{
+		Constraint->SetAngularDriveStiffness(InAngularDriveStiffness);
+	}
+
+	void FJointConstraint6DOFAsyncGT::SetAngularDriveDamping(const FVector3d& InAngularDriveDamping)
+	{
+		Constraint->SetAngularDriveDamping(InAngularDriveDamping);
+	}
+
+	void FJointConstraint6DOFAsyncGT::SetLinearDriveStiffness(const FVector3d& InLinearDriveStiffness)
+	{
+		Constraint->SetLinearDriveStiffness(InLinearDriveStiffness);
+	}
+
+	void FJointConstraint6DOFAsyncGT::SetLinearDriveDamping(const FVector3d& InLinearDriveDamping)
+	{
+		Constraint->SetLinearDriveDamping(InLinearDriveDamping);
+	}
+
+
 	void FJointConstraint6DOFAsyncGT::SetId(const FRigidObjectId& InId)
 	{
 		ConstraintId = InId;
diff --git a/Engine/Plugins/Experimental/ChaosRigidPhysicsAsync/Source/ChaosRigidPhysicsAsync/Private/ChaosRigidPhysicsAsync/RigidBodyAsync.cpp b/Engine/Plugins/Experimental/ChaosRigidPhysicsAsync/Source/ChaosRigidPhysicsAsync/Private/ChaosRigidPhysicsAsync/RigidBodyAsync.cpp
--- a/Engine/Plugins/Experimental/ChaosRigidPhysicsAsync/Source/ChaosRigidPhysicsAsync/Private/ChaosRigidPhysicsAsync/RigidBodyAsync.cpp
+++ b/Engine/Plugins/Experimental/ChaosRigidPhysicsAsync/Source/ChaosRigidPhysicsAsync/Private/ChaosRigidPhysicsAsync/RigidBodyAsync.cpp
@@ -307,6 +307,16 @@ namespace Chaos::Rigids::Async
 	//void SetGyroscopicTorqueEnabled(bool bInEnabled) {  }
 	//void SetGravityGroup(int32 InId) {  }
 
+	void FRigidBodyAsyncGT::SetCenterOfMass(const FVector3d& InCenterOfMass)
+	{
+		GetProxyAPI().SetCenterOfMass(InCenterOfMass);
+	}
+
+	void FRigidBodyAsyncGT::SetRotationOfMass(const FQuat& InRotationOfMass)
+	{
+		GetProxyAPI().SetRotationOfMass(InRotationOfMass);
+	}
+
 	void FRigidBodyAsyncGT::SetGravityScale(float InScale)
 	{ 
 		UE_RIGIDPHYSICS_CHECK(IsValid());
diff --git a/Engine/Plugins/Experimental/ChaosRigidPhysicsAsync/Source/ChaosRigidPhysicsAsync/Private/ChaosRigidPhysicsAsync/RigidSceneAsync.cpp b/Engine/Plugins/Experimental/ChaosRigidPhysicsAsync/Source/ChaosRigidPhysicsAsync/Private/ChaosRigidPhysicsAsync/RigidSceneAsync.cpp
--- a/Engine/Plugins/Experimental/ChaosRigidPhysicsAsync/Source/ChaosRigidPhysicsAsync/Private/ChaosRigidPhysicsAsync/RigidSceneAsync.cpp
+++ b/Engine/Plugins/Experimental/ChaosRigidPhysicsAsync/Source/ChaosRigidPhysicsAsync/Private/ChaosRigidPhysicsAsync/RigidSceneAsync.cpp
@@ -313,6 +313,40 @@ namespace Chaos::Rigids::Async
 		}
 	}
 
+	void FRigidSceneAsyncGT::AddDisabledCollisions(const TMap<UE::Physics::FRigidBodyHandle, TArray<UE::Physics::FRigidBodyHandle>>& InMap)
+	{
+		if (FPBDRigidsSolver* Solver = GetSolver())
+		{
+			TMap<FPhysicsActorHandle, TArray<FPhysicsActorHandle>> IgnoreMap;
+
+			for (const TPair<UE::Physics::FRigidBodyHandle, TArray<UE::Physics::FRigidBodyHandle>>& Pair : InMap)
+			{
+				if(UE::Physics::IRigidBody* Key = GetBody(Pair.Key.GetId()))
+				{ 
+					FRigidBodyAsyncGT* KeyBody = Key->AsAChecked<FRigidBodyAsyncGT>();
+
+					const TArray<UE::Physics::FRigidBodyHandle>& Values = Pair.Value;
+
+					TArray<FPhysicsActorHandle> Handles;
+					Handles.Reserve(Values.Num());
+
+					for (const UE::Physics::FRigidBodyHandle& Value : Values)
+					{
+						if(UE::Physics::IRigidBody* ValuePtr = GetBody(Value.GetId()))
+						{ 
+							FRigidBodyAsyncGT* ValueBody = ValuePtr->AsAChecked<FRigidBodyAsyncGT>();
+							Handles.Add(ValueBody->GetProxy());
+						}
+					}
+
+					IgnoreMap.Add(KeyBody->GetProxy(), Handles);
+				}
+			}
+
+			FChaosEngineInterface::AddDisabledCollisionsFor_AssumesLocked(IgnoreMap);
+		}
+	}
+
 	FRigidShapeInstanceAsync* FRigidSceneAsyncGT::AllocateShape()
 	{
 		return ShapeInstanceContainer.Allocate(SceneId);
diff --git a/Engine/Plugins/Experimental/ChaosRigidPhysicsAsync/Source/ChaosRigidPhysicsAsync/Public/ChaosRigidPhysicsAsync/JointConstraint6DOFAsync.h b/Engine/Plugins/Experimental/ChaosRigidPhysicsAsync/Source/ChaosRigidPhysicsAsync/Public/ChaosRigidPhysicsAsync/JointConstraint6DOFAsync.h
--- a/Engine/Plugins/Experimental/ChaosRigidPhysicsAsync/Source/ChaosRigidPhysicsAsync/Public/ChaosRigidPhysicsAsync/JointConstraint6DOFAsync.h
+++ b/Engine/Plugins/Experimental/ChaosRigidPhysicsAsync/Source/ChaosRigidPhysicsAsync/Public/ChaosRigidPhysicsAsync/JointConstraint6DOFAsync.h
@@ -43,6 +43,42 @@ namespace Chaos::Rigids::Async
 		UE_INTERNAL virtual EJointMotionType GetLinearMotionTypesZ() const override final;
 		UE_INTERNAL virtual void SetLinearMotionTypesZ(const EJointMotionType InMotionType) override final;
 
+		UE_INTERNAL virtual void SetAngularMotionTypesX(const EJointMotionType InMotionType) override final;
+		UE_INTERNAL virtual void SetAngularMotionTypesY(const EJointMotionType InMotionType) override final;
+		UE_INTERNAL virtual void SetAngularMotionTypesZ(const EJointMotionType InMotionType) override final;
+
+		UE_INTERNAL virtual void SetAngularLimits(const FVector3d& InAngularLimits) override final;
+
+		UE_INTERNAL virtual void SetCollisionEnabled(const bool bCollisionEnabled) override final;
+		UE_INTERNAL virtual void SetProjectionEnabled(const bool bProjectionEnabled) override final;
+		UE_INTERNAL virtual void SetMassConditioningEnabled(const bool bMassConditioningEnabled) override final;
+		UE_INTERNAL virtual void SetUseLinearSolver(const bool bUseLinearSolver) override final;
+
+		UE_INTERNAL virtual void SetSoftLinearLimitsEnabled(const bool bSoftLinearLimitsEnabled) override final;
+		UE_INTERNAL virtual void SetSoftTwistLimitsEnabled(const bool bSoftTwistLimitsEnabled) override final;
+		UE_INTERNAL virtual void SetSoftSwingLimitsEnabled(const bool bSoftSwingLimitsEnabled) override final;
+
+		UE_INTERNAL virtual void SetSoftLinearStiffness(const double InSoftLinearStiffness) override final;
+		UE_INTERNAL virtual void SetSoftLinearDamping(const double InSoftLinearDamping) override final;
+		UE_INTERNAL virtual void SetSoftTwistStiffness(const double InSoftTwistStiffness) override final;
+		UE_INTERNAL virtual void SetSoftTwistDamping(const double InSoftTwistDamping) override final;
+		UE_INTERNAL virtual void SetSoftSwingStiffness(const double InSoftSwingStiffness) override final;
+		UE_INTERNAL virtual void SetSoftSwingDamping(const double InSoftSwingDamping) override final;
+
+		UE_INTERNAL virtual void SetProjectionLinearAlpha(const double InProjectionLinearAlpha) override final;
+		UE_INTERNAL virtual void SetProjectionAngularAlpha(const double InProjectionAngularAlpha) override final;
+		UE_INTERNAL virtual void SetProjectionLinearTolerance(const double InProjectionLinearTolerance) override final;
+		UE_INTERNAL virtual void SetProjectionAngularTolerance(const double InProjectionAngularTolerance) override final;
+
+		UE_INTERNAL virtual void SetLinearContactDistance(const double InLinearContactDistance) override final;
+		UE_INTERNAL virtual void SetTwistContactDistance(const double InTwistContactDistance) override final;
+		UE_INTERNAL virtual void SetSwingContactDistance(const double InSwingContactDistance) override final;
+
+		UE_INTERNAL virtual void SetAngularDriveStiffness(const FVector3d& InAngularDriveStiffness) override final;
+		UE_INTERNAL virtual void SetAngularDriveDamping(const FVector3d& InAngularDriveDamping) override final;
+		UE_INTERNAL virtual void SetLinearDriveStiffness(const FVector3d& InLinearDriveStiffness) override final;
+		UE_INTERNAL virtual void SetLinearDriveDamping(const FVector3d& InLinearDriveDamping) override final;
+
 		UE_INTERNAL void SetId(const FRigidObjectId& InId);
 		UE_INTERNAL Chaos::FJointConstraint* GetJointConstraint();
 		UE_INTERNAL const Chaos::FJointConstraint* GetJointConstraint() const;
diff --git a/Engine/Plugins/Experimental/ChaosRigidPhysicsAsync/Source/ChaosRigidPhysicsAsync/Public/ChaosRigidPhysicsAsync/RigidBodyAsync.h b/Engine/Plugins/Experimental/ChaosRigidPhysicsAsync/Source/ChaosRigidPhysicsAsync/Public/ChaosRigidPhysicsAsync/RigidBodyAsync.h
--- a/Engine/Plugins/Experimental/ChaosRigidPhysicsAsync/Source/ChaosRigidPhysicsAsync/Public/ChaosRigidPhysicsAsync/RigidBodyAsync.h
+++ b/Engine/Plugins/Experimental/ChaosRigidPhysicsAsync/Source/ChaosRigidPhysicsAsync/Public/ChaosRigidPhysicsAsync/RigidBodyAsync.h
@@ -78,6 +78,8 @@ namespace Chaos::Rigids::Async
 		//UE_INTERNAL virtual void SetInertiaConditioningEnabled(bool bInEnabled) override final {  }
 		//UE_INTERNAL virtual void SetGyroscopicTorqueEnabled(bool bInEnabled) override final {  }
 		//UE_INTERNAL virtual void SetGravityGroup(int32 InId) override final {  }
+		UE_INTERNAL virtual void SetCenterOfMass(const FVector3d& InCenterOfMass) override final;
+		UE_INTERNAL virtual void SetRotationOfMass(const FQuat& InRotationOfMass) override final;
 		UE_INTERNAL virtual void SetGravityScale(float InScale) override final;
 		//UE_INTERNAL virtual void SetWakeEventsEnabled(bool bInEnabled) override final {  }
 		//UE_INTERNAL virtual void SetNumPositionIterations(int32 InNum) override final {  }
diff --git a/Engine/Plugins/Experimental/ChaosRigidPhysicsAsync/Source/ChaosRigidPhysicsAsync/Public/ChaosRigidPhysicsAsync/RigidSceneAsync.h b/Engine/Plugins/Experimental/ChaosRigidPhysicsAsync/Source/ChaosRigidPhysicsAsync/Public/ChaosRigidPhysicsAsync/RigidSceneAsync.h
--- a/Engine/Plugins/Experimental/ChaosRigidPhysicsAsync/Source/ChaosRigidPhysicsAsync/Public/ChaosRigidPhysicsAsync/RigidSceneAsync.h
+++ b/Engine/Plugins/Experimental/ChaosRigidPhysicsAsync/Source/ChaosRigidPhysicsAsync/Public/ChaosRigidPhysicsAsync/RigidSceneAsync.h
@@ -79,6 +79,8 @@ namespace Chaos::Rigids::Async
 		UE_INTERNAL virtual void DestroyBody(UE::Physics::IRigidBody* InBody) override final;
 		UE_INTERNAL virtual FRigidGeometryCollectionAsyncGT* CreateGeometryCollection(const UE::Physics::FRigidDebugName& InName, const FSimulationParameters& InParameters) override final;
 		UE_INTERNAL virtual void DestroyGeometryCollection(UE::Physics::IRigidGeometryCollection* InGC) override final;
+
+		UE_INTERNAL virtual void AddDisabledCollisions(const TMap<UE::Physics::FRigidBodyHandle, TArray<UE::Physics::FRigidBodyHandle>>& InMap) override final;
 		//
 		// IRIGIDSCENE END
 		//
diff --git a/Engine/Source/Runtime/Experimental/RigidPhysics/Private/RigidPhysics/Internal/IJointConstraint6DOF.cpp b/Engine/Source/Runtime/Experimental/RigidPhysics/Private/RigidPhysics/Internal/IJointConstraint6DOF.cpp
--- a/Engine/Source/Runtime/Experimental/RigidPhysics/Private/RigidPhysics/Internal/IJointConstraint6DOF.cpp
+++ b/Engine/Source/Runtime/Experimental/RigidPhysics/Private/RigidPhysics/Internal/IJointConstraint6DOF.cpp
@@ -72,6 +72,147 @@ namespace UE::Physics
 	{
 		RIGIDPHYSICS_API_UNSUPPORTED();
 	}
+
+	void IJointConstraint6DOF::SetAngularMotionTypesX(const EJointMotionType InMotionType)
+	{
+		RIGIDPHYSICS_API_UNSUPPORTED();
+	}
+
+	void IJointConstraint6DOF::SetAngularMotionTypesY(const EJointMotionType InMotionType)
+	{
+		RIGIDPHYSICS_API_UNSUPPORTED();
+	}
+
+	void IJointConstraint6DOF::SetAngularMotionTypesZ(const EJointMotionType InMotionType)
+	{
+		RIGIDPHYSICS_API_UNSUPPORTED();
+	}
+
+	void IJointConstraint6DOF::SetAngularLimits(const FVector3d& InAngularLimits)
+	{
+		RIGIDPHYSICS_API_UNSUPPORTED();
+	}
+
+	void IJointConstraint6DOF::SetCollisionEnabled(const bool bCollisionEnabled)
+	{
+		RIGIDPHYSICS_API_UNSUPPORTED();
+	}
+
+	void IJointConstraint6DOF::SetMassConditioningEnabled(const bool bMassConditioningEnabled)
+	{
+		RIGIDPHYSICS_API_UNSUPPORTED();
+	}
+
+	void IJointConstraint6DOF::SetUseLinearSolver(const bool bUseLinearSolver)
+	{
+		RIGIDPHYSICS_API_UNSUPPORTED();
+	}
+
+	void IJointConstraint6DOF::SetProjectionEnabled(const bool bProjectionEnabled)
+	{
+		RIGIDPHYSICS_API_UNSUPPORTED();
+	}
+
+	void IJointConstraint6DOF::SetSoftLinearLimitsEnabled(const bool bSoftLinearLimitsEnabled)
+	{
+		RIGIDPHYSICS_API_UNSUPPORTED();
+	}
+
+	void IJointConstraint6DOF::SetSoftTwistLimitsEnabled(const bool bSoftTwistLimitsEnabled)
+	{
+		RIGIDPHYSICS_API_UNSUPPORTED();
+	}
+
+	void IJointConstraint6DOF::SetSoftSwingLimitsEnabled(const bool bSoftSwingLimitsEnabled)
+	{
+		RIGIDPHYSICS_API_UNSUPPORTED();
+	}
+
+	void IJointConstraint6DOF::SetSoftLinearStiffness(const double InSoftLinearStiffness)
+	{
+		RIGIDPHYSICS_API_UNSUPPORTED();
+	}
+
+	void IJointConstraint6DOF::SetSoftLinearDamping(const double InSoftLinearDamping)
+	{
+		RIGIDPHYSICS_API_UNSUPPORTED();
+	}
+
+	void IJointConstraint6DOF::SetSoftTwistStiffness(const double InSoftTwistStiffness)
+	{
+		RIGIDPHYSICS_API_UNSUPPORTED();
+	}
+
+	void IJointConstraint6DOF::SetSoftTwistDamping(const double InSoftTwistDamping)
+	{
+		RIGIDPHYSICS_API_UNSUPPORTED();
+	}
+
+	void IJointConstraint6DOF::SetSoftSwingStiffness(const double InSoftSwingStiffness)
+	{
+		RIGIDPHYSICS_API_UNSUPPORTED();
+	}
+
+	void IJointConstraint6DOF::SetSoftSwingDamping(const double InSoftSwingDamping)
+	{
+		RIGIDPHYSICS_API_UNSUPPORTED();
+	}
+
+	void IJointConstraint6DOF::SetProjectionLinearAlpha(const double InProjectionLinearAlpha)
+	{
+		RIGIDPHYSICS_API_UNSUPPORTED();
+	}
+
+	void IJointConstraint6DOF::SetProjectionAngularAlpha(const double InProjectionAngularAlpha)
+	{
+		RIGIDPHYSICS_API_UNSUPPORTED();
+	}
+
+	void IJointConstraint6DOF::SetProjectionLinearTolerance(const double InProjectionLinearTolerance)
+	{
+		RIGIDPHYSICS_API_UNSUPPORTED();
+	}
+
+	void IJointConstraint6DOF::SetProjectionAngularTolerance(const double InProjectionAngularTolerance)
+	{
+		RIGIDPHYSICS_API_UNSUPPORTED();
+	}
+
+	void IJointConstraint6DOF::SetLinearContactDistance(const double InLinearContactDistance)
+	{
+		RIGIDPHYSICS_API_UNSUPPORTED();
+	}
+
+	void IJointConstraint6DOF::SetTwistContactDistance(const double InTwistContactDistance)
+	{
+		RIGIDPHYSICS_API_UNSUPPORTED();
+	}
+
+	void IJointConstraint6DOF::SetSwingContactDistance(const double InSwingContactDistance)
+	{
+		RIGIDPHYSICS_API_UNSUPPORTED();
+	}
+
+	void IJointConstraint6DOF::SetAngularDriveStiffness(const FVector3d& InAngularDriveStiffness)
+	{
+		RIGIDPHYSICS_API_UNSUPPORTED();
+	}
+
+	void IJointConstraint6DOF::SetAngularDriveDamping(const FVector3d& InAngularDriveDamping)
+	{
+		RIGIDPHYSICS_API_UNSUPPORTED();
+	}
+
+	void IJointConstraint6DOF::SetLinearDriveStiffness(const FVector3d& InLinearDriveStiffness)
+	{
+		RIGIDPHYSICS_API_UNSUPPORTED();
+	}
+
+	void IJointConstraint6DOF::SetLinearDriveDamping(const FVector3d& InLinearDriveDamping)
+	{
+		RIGIDPHYSICS_API_UNSUPPORTED();
+	}
+
 } // namespace UE::Physics
 
 #endif // UE_RIGIDPHYSICS_API_ENABLED
diff --git a/Engine/Source/Runtime/Experimental/RigidPhysics/Private/RigidPhysics/Internal/IRigidBody.cpp b/Engine/Source/Runtime/Experimental/RigidPhysics/Private/RigidPhysics/Internal/IRigidBody.cpp
--- a/Engine/Source/Runtime/Experimental/RigidPhysics/Private/RigidPhysics/Internal/IRigidBody.cpp
+++ b/Engine/Source/Runtime/Experimental/RigidPhysics/Private/RigidPhysics/Internal/IRigidBody.cpp
@@ -107,6 +107,16 @@ namespace UE::Physics
 		RIGIDPHYSICS_API_UNSUPPORTED();
 	}
 
+	void IRigidBody::SetCenterOfMass(const FVector3d& InCenterOfMass)
+	{
+		RIGIDPHYSICS_API_UNSUPPORTED();
+	}
+
+	void IRigidBody::SetRotationOfMass(const FQuat& InRotationOfMass)
+	{
+		RIGIDPHYSICS_API_UNSUPPORTED();
+	}
+
 	void IRigidBody::SetCCDEnabled(bool bInEnabled)
 	{
 		RIGIDPHYSICS_API_UNSUPPORTED();
diff --git a/Engine/Source/Runtime/Experimental/RigidPhysics/Private/RigidPhysics/Internal/IRigidScene.cpp b/Engine/Source/Runtime/Experimental/RigidPhysics/Private/RigidPhysics/Internal/IRigidScene.cpp
--- a/Engine/Source/Runtime/Experimental/RigidPhysics/Private/RigidPhysics/Internal/IRigidScene.cpp
+++ b/Engine/Source/Runtime/Experimental/RigidPhysics/Private/RigidPhysics/Internal/IRigidScene.cpp
@@ -54,6 +54,11 @@ namespace UE::Physics
 		RIGIDPHYSICS_API_UNSUPPORTED();
 	}
 
+	void IRigidScene::AddDisabledCollisions(const TMap<FRigidBodyHandle, TArray<FRigidBodyHandle>>& InMap)
+	{
+		RIGIDPHYSICS_API_UNSUPPORTED();
+	}
+
 	IJointConstraint6DOF* IRigidScene::CreateJointConstraint6DOF()
 	{
 		RIGIDPHYSICS_API_UNSUPPORTED();
diff --git a/Engine/Source/Runtime/Experimental/RigidPhysics/Private/RigidPhysics/RigidBodyHandle.cpp b/Engine/Source/Runtime/Experimental/RigidPhysics/Private/RigidPhysics/RigidBodyHandle.cpp
--- a/Engine/Source/Runtime/Experimental/RigidPhysics/Private/RigidPhysics/RigidBodyHandle.cpp
+++ b/Engine/Source/Runtime/Experimental/RigidPhysics/Private/RigidPhysics/RigidBodyHandle.cpp
@@ -29,6 +29,11 @@ namespace UE::Physics
 		BodyId.Reset();
 	}
 
+	uint32 GetTypeHash(const FRigidBodyHandle& InHandle)
+	{
+		return HashCombine(GetTypeHash(InHandle.BodyId), GetTypeHash(InHandle.GetSceneId()));
+	}
+
 	IRigidBody* FRigidBodyHandle::PinInternal(const IRigidScene* Scene) const
 	{
 		return Scene->GetBody(GetId());
diff --git a/Engine/Source/Runtime/Experimental/RigidPhysics/Private/RigidPhysics/RigidSceneId.cpp b/Engine/Source/Runtime/Experimental/RigidPhysics/Private/RigidPhysics/RigidSceneId.cpp
--- a/Engine/Source/Runtime/Experimental/RigidPhysics/Private/RigidPhysics/RigidSceneId.cpp
+++ b/Engine/Source/Runtime/Experimental/RigidPhysics/Private/RigidPhysics/RigidSceneId.cpp
@@ -18,6 +18,12 @@ namespace UE::Physics
 	{
 		Id = INDEX_NONE;
 	}
+
+	uint32 GetTypeHash(const FRigidSceneId& InId)
+	{
+		return ::GetTypeHash(InId.Id);
+	}
+
 } // namespace UE::Physics
 
 #endif // UE_RIGIDPHYSICS_API_ENABLED
diff --git a/Engine/Source/Runtime/Experimental/RigidPhysics/Public/RigidPhysics/Internal/IJointConstraint6DOF.h b/Engine/Source/Runtime/Experimental/RigidPhysics/Public/RigidPhysics/Internal/IJointConstraint6DOF.h
--- a/Engine/Source/Runtime/Experimental/RigidPhysics/Public/RigidPhysics/Internal/IJointConstraint6DOF.h
+++ b/Engine/Source/Runtime/Experimental/RigidPhysics/Public/RigidPhysics/Internal/IJointConstraint6DOF.h
@@ -39,6 +39,59 @@ namespace UE::Physics
 
 		UE_INTERNAL RIGIDPHYSICS_API virtual EJointMotionType GetLinearMotionTypesZ() const;
 		UE_INTERNAL RIGIDPHYSICS_API virtual void SetLinearMotionTypesZ(EJointMotionType InMotionType);
+
+		UE_INTERNAL RIGIDPHYSICS_API virtual void SetAngularMotionTypesX(const EJointMotionType InMotionType);
+
+		UE_INTERNAL RIGIDPHYSICS_API virtual void SetAngularMotionTypesY(const EJointMotionType InMotionType);
+
+		UE_INTERNAL RIGIDPHYSICS_API virtual void SetAngularMotionTypesZ(const EJointMotionType InMotionType);
+
+		UE_INTERNAL RIGIDPHYSICS_API virtual void SetAngularLimits(const FVector3d& InAngularLimits);
+
+		UE_INTERNAL RIGIDPHYSICS_API virtual void SetCollisionEnabled(const bool bCollisionEnabled);
+
+		UE_INTERNAL RIGIDPHYSICS_API virtual void SetProjectionEnabled(const bool bProjectionEnabled);
+
+		UE_INTERNAL RIGIDPHYSICS_API virtual void SetMassConditioningEnabled(const bool bMassConditioningEnabled);
+
+		UE_INTERNAL RIGIDPHYSICS_API virtual void SetUseLinearSolver(const bool bUseLinearSolver);
+
+		UE_INTERNAL RIGIDPHYSICS_API virtual void SetSoftLinearLimitsEnabled(const bool bSoftLinearLimitsEnabled);
+
+		UE_INTERNAL RIGIDPHYSICS_API virtual void SetSoftTwistLimitsEnabled(const bool bSoftTwistLimitsEnabled);
+
+		UE_INTERNAL RIGIDPHYSICS_API virtual void SetSoftSwingLimitsEnabled(const bool bSoftSwingLimitsEnabled);
+
+		UE_INTERNAL RIGIDPHYSICS_API virtual void SetSoftLinearStiffness(const double InSoftLinearStiffness);
+
+		UE_INTERNAL RIGIDPHYSICS_API virtual void SetSoftLinearDamping(const double InSoftLinearDamping);
+
+		UE_INTERNAL RIGIDPHYSICS_API virtual void SetSoftTwistStiffness(const double InSoftTwistStiffness);
+
+		UE_INTERNAL RIGIDPHYSICS_API virtual void SetSoftTwistDamping(const double InSoftTwistDamping);
+
+		UE_INTERNAL RIGIDPHYSICS_API virtual void SetSoftSwingStiffness(const double InSoftSwingStiffness);
+
+		UE_INTERNAL RIGIDPHYSICS_API virtual void SetSoftSwingDamping(const double InSoftSwingDamping);
+
+		UE_INTERNAL RIGIDPHYSICS_API virtual void SetProjectionLinearAlpha(const double InProjectionLinearAlpha);
+
+		UE_INTERNAL RIGIDPHYSICS_API virtual void SetProjectionAngularAlpha(const double InProjectionAngularAlpha);
+
+		UE_INTERNAL RIGIDPHYSICS_API virtual void SetProjectionLinearTolerance(const double InProjectionLinearTolerance);
+
+		UE_INTERNAL RIGIDPHYSICS_API virtual void SetProjectionAngularTolerance(const double InProjectionAngularTolerance);
+
+		UE_INTERNAL RIGIDPHYSICS_API virtual void SetLinearContactDistance(const double InLinearContactDistance);
+
+		UE_INTERNAL RIGIDPHYSICS_API virtual void SetTwistContactDistance(const double InTwistContactDistance);
+
+		UE_INTERNAL RIGIDPHYSICS_API virtual void SetSwingContactDistance(const double InSwingContactDistance);
+
+		UE_INTERNAL RIGIDPHYSICS_API virtual void SetAngularDriveStiffness(const FVector3d& InAngularDriveStiffness);
+		UE_INTERNAL RIGIDPHYSICS_API virtual void SetAngularDriveDamping(const FVector3d& InAngularDriveDamping);
+		UE_INTERNAL RIGIDPHYSICS_API virtual void SetLinearDriveStiffness(const FVector3d& InLinearDriveStiffness);
+		UE_INTERNAL RIGIDPHYSICS_API virtual void SetLinearDriveDamping(const FVector3d& InLinearDriveDamping);
 	};
 } // namespace UE::Physics
 
diff --git a/Engine/Source/Runtime/Experimental/RigidPhysics/Public/RigidPhysics/Internal/IRigidBody.h b/Engine/Source/Runtime/Experimental/RigidPhysics/Public/RigidPhysics/Internal/IRigidBody.h
--- a/Engine/Source/Runtime/Experimental/RigidPhysics/Public/RigidPhysics/Internal/IRigidBody.h
+++ b/Engine/Source/Runtime/Experimental/RigidPhysics/Public/RigidPhysics/Internal/IRigidBody.h
@@ -90,6 +90,9 @@ namespace UE::Physics
 		UE_INTERNAL RIGIDPHYSICS_API virtual FVector3d GetInertia() const;
 		UE_INTERNAL RIGIDPHYSICS_API virtual void SetInertia(const FVector3d& InInertia);
 
+		UE_INTERNAL RIGIDPHYSICS_API virtual void SetCenterOfMass(const FVector3d& InCenterOfMass);
+		UE_INTERNAL RIGIDPHYSICS_API virtual void SetRotationOfMass(const FQuat& InRotationOfMass);
+
 		UE_INTERNAL RIGIDPHYSICS_API virtual void SetCCDEnabled(bool bInEnabled);
 
 		UE_INTERNAL RIGIDPHYSICS_API virtual void SetMACDEnabled(bool bInEnabled);
diff --git a/Engine/Source/Runtime/Experimental/RigidPhysics/Public/RigidPhysics/Internal/IRigidScene.h b/Engine/Source/Runtime/Experimental/RigidPhysics/Public/RigidPhysics/Internal/IRigidScene.h
--- a/Engine/Source/Runtime/Experimental/RigidPhysics/Public/RigidPhysics/Internal/IRigidScene.h
+++ b/Engine/Source/Runtime/Experimental/RigidPhysics/Public/RigidPhysics/Internal/IRigidScene.h
@@ -49,6 +49,8 @@ namespace UE::Physics
 		UE_INTERNAL RIGIDPHYSICS_API virtual IRigidGeometryCollection* CreateGeometryCollection(const FRigidDebugName& InName, const FSimulationParameters& InParameters);
 		UE_INTERNAL RIGIDPHYSICS_API virtual void DestroyGeometryCollection(IRigidGeometryCollection* InGC);
 
+		UE_INTERNAL RIGIDPHYSICS_API virtual void AddDisabledCollisions(const TMap<FRigidBodyHandle, TArray<FRigidBodyHandle>>& InMap);
+
 		UE_INTERNAL RIGIDPHYSICS_API virtual IJointConstraint6DOF* CreateJointConstraint6DOF();
 		UE_INTERNAL RIGIDPHYSICS_API virtual void DestroyJointConstraint(IJointConstraint* InConstraint);
 		UE_INTERNAL RIGIDPHYSICS_API virtual IJointConstraint* GetJointConstraint(const FRigidObjectId& InConstraintId) const;
diff --git a/Engine/Source/Runtime/Experimental/RigidPhysics/Public/RigidPhysics/JointConstraint6DOF.h b/Engine/Source/Runtime/Experimental/RigidPhysics/Public/RigidPhysics/JointConstraint6DOF.h
--- a/Engine/Source/Runtime/Experimental/RigidPhysics/Public/RigidPhysics/JointConstraint6DOF.h
+++ b/Engine/Source/Runtime/Experimental/RigidPhysics/Public/RigidPhysics/JointConstraint6DOF.h
@@ -117,6 +117,146 @@ namespace UE::Physics
 			ConstraintRaw->SetLinearMotionTypesZ(InMotionType);
 		}
 
+		void SetAngularMotionTypesX(const EJointMotionType InMotionType) requires CIsGameContext<FContext>
+		{
+			ConstraintRaw->SetAngularMotionTypesX(InMotionType);
+		}
+
+		void SetAngularMotionTypesY(const EJointMotionType InMotionType) requires CIsGameContext<FContext>
+		{
+			ConstraintRaw->SetAngularMotionTypesY(InMotionType);
+		}
+
+		void SetAngularMotionTypesZ(const EJointMotionType InMotionType) requires CIsGameContext<FContext>
+		{
+			ConstraintRaw->SetAngularMotionTypesZ(InMotionType);
+		}
+
+		void SetAngularLimits(const FVector3d& InAngularLimits) requires CIsGameContext<FContext>
+		{
+			ConstraintRaw->SetAngularLimits(InAngularLimits);
+		}
+
+		void SetCollisionEnabled(const bool bCollisionEnabled) requires CIsGameContext<FContext>
+		{
+			ConstraintRaw->SetCollisionEnabled(bCollisionEnabled);
+		}
+
+		void SetProjectionEnabled(const bool bProjectionEnabled) requires CIsGameContext<FContext>
+		{
+			ConstraintRaw->SetProjectionEnabled(bProjectionEnabled);
+		}
+
+		void SetMassConditioningEnabled(const bool bMassConditioningEnabled) requires CIsGameContext<FContext>
+		{
+			ConstraintRaw->SetMassConditioningEnabled(bMassConditioningEnabled);
+		}
+
+		void SetUseLinearSolver(const bool bUseLinearSolver) requires CIsGameContext<FContext>
+		{
+			ConstraintRaw->SetUseLinearSolver(bUseLinearSolver);
+		}
+
+		void SetSoftLinearLimitsEnabled(const bool bSoftLinearLimitsEnabled) requires CIsGameContext<FContext>
+		{
+			ConstraintRaw->SetSoftLinearLimitsEnabled(bSoftLinearLimitsEnabled);
+		}
+
+		void SetSoftTwistLimitsEnabled(const bool bSoftTwistLimitsEnabled) requires CIsGameContext<FContext>
+		{
+			ConstraintRaw->SetSoftTwistLimitsEnabled(bSoftTwistLimitsEnabled);
+		}
+
+		void SetSoftSwingLimitsEnabled(const bool bSoftSwingLimitsEnabled) requires CIsGameContext<FContext>
+		{
+			ConstraintRaw->SetSoftSwingLimitsEnabled(bSoftSwingLimitsEnabled);
+		}
+
+		void SetSoftLinearStiffness(const double InSoftLinearStiffness) requires CIsGameContext<FContext>
+		{
+			ConstraintRaw->SetSoftLinearStiffness(InSoftLinearStiffness);
+		}
+
+		void SetSoftLinearDamping(const double InSoftLinearDamping) requires CIsGameContext<FContext>
+		{
+			ConstraintRaw->SetSoftLinearDamping(InSoftLinearDamping);
+		}
+
+		void SetSoftTwistStiffness(const double InSoftTwistStiffness) requires CIsGameContext<FContext>
+		{
+			ConstraintRaw->SetSoftTwistStiffness(InSoftTwistStiffness);
+		}
+
+		void SetSoftTwistDamping(const double InSoftTwistDamping) requires CIsGameContext<FContext>
+		{
+			ConstraintRaw->SetSoftTwistDamping(InSoftTwistDamping);
+		}	
+		
+		void SetSoftSwingStiffness(const double InSoftSwingStiffness) requires CIsGameContext<FContext>
+		{
+			ConstraintRaw->SetSoftSwingStiffness(InSoftSwingStiffness);
+		}
+
+		void SetSoftSwingDamping(const double InSoftSwingDamping) requires CIsGameContext<FContext>
+		{
+			ConstraintRaw->SetSoftSwingDamping(InSoftSwingDamping);
+		}
+
+		void SetProjectionLinearAlpha(const double InProjectionLinearAlpha) requires CIsGameContext<FContext>
+		{
+			ConstraintRaw->SetProjectionLinearAlpha(InProjectionLinearAlpha);
+		}
+
+		void SetProjectionAngularAlpha(const double InProjectionAngularAlpha) requires CIsGameContext<FContext>
+		{
+			ConstraintRaw->SetProjectionAngularAlpha(InProjectionAngularAlpha);
+		}
+
+		void SetProjectionLinearTolerance(const double InProjectionLinearTolerance) requires CIsGameContext<FContext>
+		{
+			ConstraintRaw->SetProjectionLinearTolerance(InProjectionLinearTolerance);
+		}
+
+		void SetProjectionAngularTolerance(const double InProjectionAngularTolerance) requires CIsGameContext<FContext>
+		{
+			ConstraintRaw->SetProjectionAngularTolerance(InProjectionAngularTolerance);
+		}
+
+		void SetLinearContactDistance(const double InLinearContactDistance) requires CIsGameContext<FContext>
+		{
+			ConstraintRaw->SetLinearContactDistance(InLinearContactDistance);
+		}
+
+		void SetTwistContactDistance(const double InTwistContactDistance) requires CIsGameContext<FContext>
+		{
+			ConstraintRaw->SetTwistContactDistance(InTwistContactDistance);
+		}
+
+		void SetSwingContactDistance(const double InSwingContactDistance) requires CIsGameContext<FContext>
+		{
+			ConstraintRaw->SetSwingContactDistance(InSwingContactDistance);
+		}
+
+		void SetAngularDriveStiffness(const FVector3d& InAngularDriveStiffness) requires CIsGameContext<FContext>
+		{
+			ConstraintRaw->SetAngularDriveStiffness(InAngularDriveStiffness);
+		}
+
+		void SetAngularDriveDamping(const FVector3d& InAngularDriveDamping) requires CIsGameContext<FContext>
+		{
+			ConstraintRaw->SetAngularDriveDamping(InAngularDriveDamping);
+		}
+
+		void SetLinearDriveStiffness(const FVector3d& InLinearDriveStiffness) requires CIsGameContext<FContext>
+		{
+			ConstraintRaw->SetLinearDriveStiffness(InLinearDriveStiffness);
+		}
+
+		void SetLinearDriveDamping(const FVector3d& InLinearDriveDamping) requires CIsGameContext<FContext>
+		{
+			ConstraintRaw->SetLinearDriveDamping(InLinearDriveDamping);
+		}
+
 		UE_INTERNAL void Reset()
 		{
 			ConstraintRaw = nullptr;
diff --git a/Engine/Source/Runtime/Experimental/RigidPhysics/Public/RigidPhysics/RigidBody.h b/Engine/Source/Runtime/Experimental/RigidPhysics/Public/RigidPhysics/RigidBody.h
--- a/Engine/Source/Runtime/Experimental/RigidPhysics/Public/RigidPhysics/RigidBody.h
+++ b/Engine/Source/Runtime/Experimental/RigidPhysics/Public/RigidPhysics/RigidBody.h
@@ -163,6 +163,16 @@ namespace UE::Physics
 			BodyRaw->SetInertia(InInertia);
 		}
 
+		void SetCenterOfMass(const FVector3d& InCenterOfMass) requires CIsGameContext<FContext>
+		{
+			BodyRaw->SetCenterOfMass(InCenterOfMass);
+		}
+
+		void SetRotationOfMass(const FQuat& InRotationOfMass) requires CIsGameContext<FContext>
+		{
+			BodyRaw->SetRotationOfMass(InRotationOfMass);
+		}
+
 		void SetCCDEnabled(bool bInEnabled)
 		{
 			BodyRaw->SetCCDEnabled(bInEnabled);
diff --git a/Engine/Source/Runtime/Experimental/RigidPhysics/Public/RigidPhysics/RigidBodyHandle.h b/Engine/Source/Runtime/Experimental/RigidPhysics/Public/RigidPhysics/RigidBodyHandle.h
--- a/Engine/Source/Runtime/Experimental/RigidPhysics/Public/RigidPhysics/RigidBodyHandle.h
+++ b/Engine/Source/Runtime/Experimental/RigidPhysics/Public/RigidPhysics/RigidBodyHandle.h
@@ -44,6 +44,8 @@ namespace UE::Physics
 		UE_INTERNAL RIGIDPHYSICS_API const FRigidSceneId& GetSceneId() const;
 		UE_INTERNAL RIGIDPHYSICS_API void Reset();
 
+		RIGIDPHYSICS_API friend uint32 GetTypeHash(const FRigidBodyHandle& InHandle);
+
 	private:
 		RIGIDPHYSICS_API IRigidBody* PinInternal(const IRigidScene* Scene) const;
 
diff --git a/Engine/Source/Runtime/Experimental/RigidPhysics/Public/RigidPhysics/RigidScene.h b/Engine/Source/Runtime/Experimental/RigidPhysics/Public/RigidPhysics/RigidScene.h
--- a/Engine/Source/Runtime/Experimental/RigidPhysics/Public/RigidPhysics/RigidScene.h
+++ b/Engine/Source/Runtime/Experimental/RigidPhysics/Public/RigidPhysics/RigidScene.h
@@ -89,6 +89,11 @@ namespace UE::Physics
 			InGC.Reset();
 		}
 
+		void AddDisabledCollisions(const TMap<FRigidBodyHandle, TArray<FRigidBodyHandle>>& InMap) requires CIsGameContext<FContext>
+		{
+			SceneRaw->AddDisabledCollisions(InMap);
+		}
+
 		TJointConstraint6DOFPtr<FContext> CreateJointConstraint6DOF() requires CIsGameContext<FContext>
 		{
 			return SceneRaw->CreateJointConstraint6DOF();
diff --git a/Engine/Source/Runtime/Experimental/RigidPhysics/Public/RigidPhysics/RigidSceneId.h b/Engine/Source/Runtime/Experimental/RigidPhysics/Public/RigidPhysics/RigidSceneId.h
--- a/Engine/Source/Runtime/Experimental/RigidPhysics/Public/RigidPhysics/RigidSceneId.h
+++ b/Engine/Source/Runtime/Experimental/RigidPhysics/Public/RigidPhysics/RigidSceneId.h
@@ -14,6 +14,8 @@ namespace UE::Physics
 		bool operator==(const FRigidSceneId&) const = default;
 		bool operator!=(const FRigidSceneId&) const = default;
 
+		RIGIDPHYSICS_API friend uint32 GetTypeHash(const FRigidSceneId& InId);
+
 		UE_INTERNAL RIGIDPHYSICS_API IRigidScene* Get() const;
 		UE_INTERNAL RIGIDPHYSICS_API void Reset();
 
