Chaos API: Inlining the creation of shapes. This is removes a dependency and is necessary for future work.

#Synced-CL 52060956

diff --git a/Engine/Plugins/Experimental/ChaosRigidPhysicsAsync/Source/ChaosRigidPhysicsAsync/Private/ChaosRigidPhysicsAsync/Conversions.h b/Engine/Plugins/Experimental/ChaosRigidPhysicsAsync/Source/ChaosRigidPhysicsAsync/Private/ChaosRigidPhysicsAsync/Conversions.h
new file mode 100644
--- /dev/null
+++ b/Engine/Plugins/Experimental/ChaosRigidPhysicsAsync/Source/ChaosRigidPhysicsAsync/Private/ChaosRigidPhysicsAsync/Conversions.h
@@ -0,0 +1,36 @@
+// Copyright Epic Games, Inc. All Rights Reserved.
+#pragma once
+
+#include "RigidPhysics/RigidFwd.h"
+
+#if UE_RIGIDPHYSICS_API_ENABLED
+
+#include "BodySetupEnums.h"
+#include "Chaos/Particles.h"
+
+namespace Chaos::Rigids::Async
+{
+	inline Chaos::EChaosCollisionTraceFlag ConvertCollisionTraceFlag(const ECollisionTraceFlag InFlag)
+	{
+		switch (InFlag)
+		{
+			case ECollisionTraceFlag::CTF_UseDefault:
+				return Chaos::EChaosCollisionTraceFlag::Chaos_CTF_UseDefault;
+			case ECollisionTraceFlag::CTF_UseSimpleAndComplex:
+				return Chaos::EChaosCollisionTraceFlag::Chaos_CTF_UseSimpleAndComplex;
+			case ECollisionTraceFlag::CTF_UseSimpleAsComplex:
+				return Chaos::EChaosCollisionTraceFlag::Chaos_CTF_UseSimpleAsComplex;
+			case ECollisionTraceFlag::CTF_UseComplexAsSimple:
+				return Chaos::EChaosCollisionTraceFlag::Chaos_CTF_UseComplexAsSimple;
+			case ECollisionTraceFlag::CTF_MAX:
+				return Chaos::EChaosCollisionTraceFlag::Chaos_CTF_MAX;
+			default:
+			{
+				ensure(false);
+				return Chaos::EChaosCollisionTraceFlag::Chaos_CTF_UseDefault;
+			}
+		}
+	}
+} // namespace Chaos::Rigids::Async
+
+#endif // UE_RIGIDPHYSICS_API_ENABLED
diff --git a/Engine/Plugins/Experimental/ChaosRigidPhysicsAsync/Source/ChaosRigidPhysicsAsync/Private/ChaosRigidPhysicsAsync/RigidShapeInstanceAsync.cpp b/Engine/Plugins/Experimental/ChaosRigidPhysicsAsync/Source/ChaosRigidPhysicsAsync/Private/ChaosRigidPhysicsAsync/RigidShapeInstanceAsync.cpp
--- a/Engine/Plugins/Experimental/ChaosRigidPhysicsAsync/Source/ChaosRigidPhysicsAsync/Private/ChaosRigidPhysicsAsync/RigidShapeInstanceAsync.cpp
+++ b/Engine/Plugins/Experimental/ChaosRigidPhysicsAsync/Source/ChaosRigidPhysicsAsync/Private/ChaosRigidPhysicsAsync/RigidShapeInstanceAsync.cpp
@@ -1,99 +1,145 @@
 // Copyright Epic Games, Inc. All Rights Reserved.
 #include "ChaosRigidPhysicsAsync/RigidShapeInstanceAsync.h"
 
-
 #if UE_RIGIDPHYSICS_API_ENABLED
 
+#include "Chaos/Box.h"
+#include "Chaos/Capsule.h"
 #include "Chaos/Convex.h"
+#include "Chaos/ImplicitObjectTransformed.h"
 #include "Chaos/ShapeInstance.h"
+#include "Chaos/Sphere.h"
 #include "Chaos/TriangleMeshImplicitObject.h"
-
-#include "Physics/Experimental/ChaosInterfaceUtils.h"
-#include "PhysicsEngine/AggregateGeom.h"
+#include "ChaosRigidPhysicsAsync/Conversions.h"
 
 namespace Chaos::Rigids::Async
 {
+	using FBoxGeometry = UE::Physics::FBoxGeometry;
+	using FCapsuleGeometry = UE::Physics::FCapsuleGeometry;
+	using FSphereGeometry = UE::Physics::FSphereGeometry;
+	using FConvexGeometry = UE::Physics::FConvexGeometry;
+	using FTriangleMeshGeometry = UE::Physics::FTriangleMeshGeometry;
+	using FAnyGeometry = UE::Physics::FAnyGeometry;
+	using FRigidShapeInstanceSetup = UE::Physics::FRigidShapeInstanceSetup;
+
 	UE_RIGIDPHYSICS_RIGIDTYPED_IMPL(FRigidShapeInstanceAsync, UE::Physics::IRigidShapeInstance);
 
-	Chaos::FPerShapeData* FRigidShapeInstanceAsync::Create(const UE::Physics::FRigidShapeInstanceSetup& InSetup)
+	Chaos::FImplicitObjectPtr MakeImplicit(const FBoxGeometry& Geometry, const FRigidShapeInstanceSetup& InSetup)
 	{
-		const UE::Physics::FAnyGeometry& InGeometry = InSetup.Geometry;
-		FKAggregateGeom AggGeom;
-		InGeometry.Visit([&AggGeom]<typename GeometryType>(const GeometryType& Geometry)
-		{
-			if constexpr (std::is_same_v<GeometryType, UE::Physics::FBoxGeometry>)
-			{
-				const FVector3f Extents = Geometry.GetExtents();
-				FKBoxElem BoxElem = FKBoxElem(Extents.X, Extents.Y, Extents.Z);
-				BoxElem.Center = FVector::Zero();
-				AggGeom.AddElement(BoxElem);
-			}
-			else if constexpr (std::is_same_v<GeometryType, UE::Physics::FSphereGeometry>)
-			{
-				FKSphereElem SphereElem = FKSphereElem(Geometry.GetRadius());
-				SphereElem.Center = FVector::Zero();
-				AggGeom.AddElement(SphereElem);
-			}
-			else if constexpr (std::is_same_v<GeometryType, UE::Physics::FCapsuleGeometry>)
-			{
-				const float Radius = Geometry.GetRadius();
-				const float Length = Geometry.GetSegmentHalfLength() * 2;
-				FKSphylElem CapsuleElement = FKSphylElem(Radius, Length);
-				CapsuleElement.Center = FVector::Zero();
-				AggGeom.AddElement(CapsuleElement);
-			}
-			else if constexpr (std::is_same_v<GeometryType, UE::Physics::FConvexGeometry>)
-			{
-				ensureMsgf(false, TEXT("TODO"));
-			}
-			else if constexpr (std::is_same_v<GeometryType, UE::Physics::FTriangleMeshGeometry>)
-			{
-				ensureMsgf(false, TEXT("TODO"));
-			}
-			else
-			{
-				static_assert(false, "Unsupported shape type");
-			}
-		});
-
-		// Not used - filters are set up after
-		FBodyCollisionData BodyCollisionData;
-
-		FGeometryAddParams AddParams;
-		AddParams.bDoubleSided = false;
-		AddParams.CollisionData = BodyCollisionData;
-		AddParams.CollisionTraceType = ECollisionTraceFlag::CTF_UseSimpleAsComplex;
-		AddParams.Scale = FVector3d::One();
-		AddParams.SimpleMaterial = nullptr;
-		AddParams.ComplexMaterials = {};
-		AddParams.ComplexMaterialMasks = {};
-		AddParams.LocalTransform = FTransform(InSetup.LocalTransform);
-		AddParams.WorldTransform = FTransform::Identity;
-		AddParams.Geometry = &AggGeom;
-
-		TArray<Chaos::FImplicitObjectPtr> Geoms;
-		Chaos::FShapesArray Shapes;
-		ChaosInterface::CreateGeometry(AddParams, Geoms, Shapes);
-
-		if (Shapes.Num() != 1)
+		const FTransform3f& InLocalTransform = InSetup.LocalTransform;
+		const FVector InScale(InLocalTransform.GetScale3D());
+		const FQuat InRotation(InLocalTransform.GetRotation());
+		const FVector InTranslation(InLocalTransform.GetTranslation());
+
+		float MinScale, MinScaleAbs;
+		FVector Scale3DAbs;
+		SetupNonUniformHelper(InScale, MinScale, MinScaleAbs, Scale3DAbs);
+
+		const FVector3f ScaledExtents = Geometry.GetExtents() * FVector3f(Scale3DAbs);
+		Chaos::FVec3 HalfExtents = ScaledExtents * 0.5;
+		HalfExtents.X = FMath::Max(HalfExtents.X, UE_KINDA_SMALL_NUMBER);
+		HalfExtents.Y = FMath::Max(HalfExtents.Y, UE_KINDA_SMALL_NUMBER);
+		HalfExtents.Z = FMath::Max(HalfExtents.Z, UE_KINDA_SMALL_NUMBER);
+
+		if (!InRotation.IsIdentity())
 		{
-			UE_LOGFMT(LogRigidPhysics, Warning, "MakeShapeInstanceFromAddParams is intended to be used for single shapes only. Have {0} shapes.", Shapes.Num());
+			const FTransform LocalTransform(InRotation, InTranslation);
+			Chaos::FImplicitObjectPtr BoxImpl = MakeImplicitObjectPtr<Chaos::FImplicitBox3>(-HalfExtents, HalfExtents);
+			return MakeImplicitObjectPtr<Chaos::FImplicitObjectTransformed>(MoveTemp(BoxImpl), LocalTransform);
 		}
+		const Chaos::FVec3 Center(InTranslation);
+		return MakeImplicitObjectPtr<Chaos::FImplicitBox3>(Center - HalfExtents, Center + HalfExtents);
+	}
+
+	Chaos::FImplicitObjectPtr MakeImplicit(const FSphereGeometry& Geometry, const FRigidShapeInstanceSetup& InSetup)
+	{
+		const FTransform3f& InLocalTransform = InSetup.LocalTransform;
+		const FVector InScale(InLocalTransform.GetScale3D());
+
+		float MinScale, MinScaleAbs;
+		FVector Scale3DAbs;
+		SetupNonUniformHelper(InScale, MinScale, MinScaleAbs, Scale3DAbs);
 
-		Chaos::FPerShapeData* PerShapeData = nullptr;
-		if (Shapes.Num() > 0)
+		const float ScaledRadius = Geometry.GetRadius() * MinScaleAbs;
+		const Chaos::FVec3 Center(InLocalTransform.GetTranslation());
+		return MakeImplicitObjectPtr<Chaos::FSphere>(Center, ScaledRadius);
+	}
+
+	Chaos::FImplicitObjectPtr MakeImplicit(const FCapsuleGeometry& Geometry, const FRigidShapeInstanceSetup& InSetup)
+	{
+		const FTransform3f& InLocalTransform = InSetup.LocalTransform;
+		const FQuat InRotation(InLocalTransform.GetRotation());
+		const FVector InTranslation(InLocalTransform.GetTranslation());
+
+		float MinScale, MinScaleAbs;
+		FVector Scale3DAbs;
+		SetupNonUniformHelper(FVector(InLocalTransform.GetScale3D()), MinScale, MinScaleAbs, Scale3DAbs);
+
+		const float Radius = Geometry.GetRadius();
+		const float SegmentHalfLen = Geometry.GetSegmentHalfLength();
+		const float RadiusScale = FMath::Max(Scale3DAbs.X, Scale3DAbs.Y);
+		const float LengthScale = Scale3DAbs.Z;
+
+		const float ScaledHalfLength = FMath::Max((Radius + SegmentHalfLen) * LengthScale, 0.1);
+		const float ScaledRadius = FMath::Clamp(RadiusScale * Radius, 0.1f, ScaledHalfLength);
+		const float FinalSegmentHalfLen = FMath::Max(0.05f, ScaledHalfLength - ScaledRadius);
+
+		const Chaos::FVec3 Axis = InRotation.RotateVector(FVector::ZAxisVector);
+		const Chaos::FVec3 Center(InTranslation);
+		const Chaos::FVec3 X1 = Center - Axis * FinalSegmentHalfLen;
+		const Chaos::FVec3 X2 = Center + Axis * FinalSegmentHalfLen;
+		return MakeImplicitObjectPtr<Chaos::FCapsule>(X1, X2, ScaledRadius);
+	}
+
+	Chaos::FImplicitObjectPtr MakeImplicit(const FConvexGeometry& Geometry, const FRigidShapeInstanceSetup& InSetup)
+	{
+		ensureMsgf(false, TEXT("TODO"));
+		return nullptr;
+	}
+
+	Chaos::FImplicitObjectPtr MakeImplicit(const FTriangleMeshGeometry& Geometry, const FRigidShapeInstanceSetup& InSetup)
+	{
+		ensureMsgf(false, TEXT("TODO"));
+		return nullptr;
+	}
+
+	template <typename GeometryType>
+	Chaos::FImplicitObjectPtr MakeImplicit(const GeometryType& Geometry, const FRigidShapeInstanceSetup& InSetup)
+	{
+		// Note: Using GeometryType in the static assert so it's unable to be evaluated until template instantiation.
+		static_assert(std::is_same_v<GeometryType, void>, "Unsupported shape type");
+		return nullptr;
+	}
+
+	Chaos::FPerShapeData* FRigidShapeInstanceAsync::Create(const FRigidShapeInstanceSetup& InSetup)
+	{
+		const FAnyGeometry& InGeometry = InSetup.Geometry;
+
+		Chaos::FImplicitObjectPtr ImplicitPtr = InGeometry.Visit([&InSetup]<typename GeometryType>(const GeometryType& Geometry)
 		{
-			PerShapeData = Shapes[0].Release();
-		}
+			return MakeImplicit(Geometry, InSetup);
+		});
 
-		if (PerShapeData != nullptr)
+		if (!ensure(ImplicitPtr))
 		{
-			const Chaos::EFilterFlags FilterFlags = Chaos::EFilterFlags::SimpleCollision | Chaos::EFilterFlags::ContactNotify | Chaos::EFilterFlags::ModifyContacts;
-			const Chaos::Filter::FShapeFilterData ShapeFilter = Chaos::Filter::FShapeFilterBuilder::BuildBlockAll(FilterFlags);
-			PerShapeData->SetQueryEnabled(true);
-			PerShapeData->SetSimEnabled(true);
-			PerShapeData->SetShapeFilterData(ShapeFilter);
+			return nullptr;
 		}
+
+		TUniquePtr<Chaos::FPerShapeData> NewShape = Chaos::FShapeInstanceProxy::Make(0, ImplicitPtr);
+		Chaos::FPerShapeData* PerShapeData = NewShape.Release();
+
+		const Chaos::EFilterFlags FilterFlags = Chaos::EFilterFlags::SimpleCollision | Chaos::EFilterFlags::ContactNotify | Chaos::EFilterFlags::ModifyContacts;
+		const Chaos::Filter::FShapeFilterData ShapeFilter = Chaos::Filter::FShapeFilterBuilder::BuildBlockAll(FilterFlags);
+
+		PerShapeData->SetQueryEnabled(true);
+		PerShapeData->SetSimEnabled(true);
+		PerShapeData->SetIsProbe(false);
+		PerShapeData->SetShapeFilterData(ShapeFilter);
+		PerShapeData->SetCollisionTraceType(ConvertCollisionTraceFlag(ECollisionTraceFlag::CTF_UseSimpleAsComplex));
+		// TODO: UpdateShapeBounds?
+		// TODO: UserData
+		// TODO: Materials
+
 		return PerShapeData;
 	}
 
@@ -105,11 +151,6 @@ namespace Chaos::Rigids::Async
 	UE::Physics::FAnyGeometry FRigidShapeInstanceAsync::GetGeometry() const
 	{
 		using FImplicitTransformed = TImplicitObjectTransformed<FReal, 3>;
-		using FBoxGeometry = UE::Physics::FBoxGeometry;
-		using FSphereGeometry = UE::Physics::FSphereGeometry;
-		using FCapsuleGeometry = UE::Physics::FCapsuleGeometry;
-		using FAnyGeometry = UE::Physics::FAnyGeometry;
-
 		const FImplicitObject* ImplicitGeom = PerShapeData->GetGeometry();
 		EImplicitObjectType ImplicitType = ImplicitGeom->GetType();
 		// TODO: Looks like there's other "wrapper" types like instanced and scaled that should be handled.
