Chaos API: Updating shape instance to handle convex, triangle mesh, and height field geometry types.
#Synced-CL 52100670

diff --git a/Engine/Plugins/Experimental/ChaosRigidPhysicsAsync/Source/ChaosRigidPhysicsAsync/Private/ChaosRigidPhysicsAsync/RigidShapeInstanceAsync.cpp b/Engine/Plugins/Experimental/ChaosRigidPhysicsAsync/Source/ChaosRigidPhysicsAsync/Private/ChaosRigidPhysicsAsync/RigidShapeInstanceAsync.cpp
--- a/Engine/Plugins/Experimental/ChaosRigidPhysicsAsync/Source/ChaosRigidPhysicsAsync/Private/ChaosRigidPhysicsAsync/RigidShapeInstanceAsync.cpp
+++ b/Engine/Plugins/Experimental/ChaosRigidPhysicsAsync/Source/ChaosRigidPhysicsAsync/Private/ChaosRigidPhysicsAsync/RigidShapeInstanceAsync.cpp
@@ -10,6 +10,7 @@
 #include "Chaos/ImplicitObjectTransformed.h"
 #include "Chaos/ShapeInstance.h"
 #include "Chaos/Sphere.h"
+#include "Chaos/HeightField.h"
 #include "Chaos/TriangleMeshImplicitObject.h"
 #include "ChaosRigidPhysicsAsync/Conversions.h"
 
@@ -20,8 +21,10 @@ namespace Chaos::Rigids::Async
 	using FSphereGeometry = UE::Physics::FSphereGeometry;
 	using FConvexGeometry = UE::Physics::FConvexGeometry;
 	using FTriangleMeshGeometry = UE::Physics::FTriangleMeshGeometry;
+	using FHeightFieldGeometry = UE::Physics::FHeightFieldGeometry;
 	using FAnyGeometry = UE::Physics::FAnyGeometry;
 	using FRigidShapeInstanceSetup = UE::Physics::FRigidShapeInstanceSetup;
+	using FImplicitTransformed = Chaos::TImplicitObjectTransformed<FReal, 3>;
 
 	UE_RIGIDPHYSICS_RIGIDTYPED_IMPL(FRigidShapeInstanceAsync, UE::Physics::IRigidShapeInstance);
 
@@ -94,14 +97,118 @@ namespace Chaos::Rigids::Async
 
 	Chaos::FImplicitObjectPtr MakeImplicit(const FConvexGeometry& Geometry, const FRigidShapeInstanceSetup& InSetup)
 	{
-		ensureMsgf(false, TEXT("TODO"));
-		return nullptr;
+		const Chaos::FConvexPtr ConvexImplicit = Geometry.GetImplicit();
+		if (!ensure(ConvexImplicit))
+		{
+			return nullptr;
+		}
+
+		const FTransform3f& InLocalTransform = InSetup.LocalTransform;
+		const FVector InScale(InLocalTransform.GetScale3D());
+		const FQuat InRotation(InLocalTransform.GetRotation());
+		const FVector InTranslation(InLocalTransform.GetTranslation());
+
+		// Clamp near-zero scale components to avoid degenerate scaled shapes
+		FVector NetScale = InScale;
+		NetScale.X = FMath::Abs(NetScale.X) < UE_KINDA_SMALL_NUMBER ? UE_KINDA_SMALL_NUMBER : NetScale.X;
+		NetScale.Y = FMath::Abs(NetScale.Y) < UE_KINDA_SMALL_NUMBER ? UE_KINDA_SMALL_NUMBER : NetScale.Y;
+		NetScale.Z = FMath::Abs(NetScale.Z) < UE_KINDA_SMALL_NUMBER ? UE_KINDA_SMALL_NUMBER : NetScale.Z;
+
+		// Extract the scale from the transform - we have separate wrapper classes for scale versus translate/rotate
+		const FTransform ConvexTransform(InRotation, InTranslation, FVector(1, 1, 1));
+		const bool bHasTranslationOrRotation = !ConvexTransform.GetTranslation().IsNearlyZero() || !ConvexTransform.GetRotation().IsIdentity();
+		const bool bNoScale = FVector::PointsAreNear(NetScale, FVector(1), UE_KINDA_SMALL_NUMBER);
+
+		// Wrap in Instanced (identity scale) or Scaled depending on scale value
+		// NOTE: margin lives on the wrapper, not the inner shared FConvex object
+		Chaos::FImplicitObjectPtr Implicit;
+		if (bNoScale)
+		{
+			Implicit = MakeImplicitObjectPtr<Chaos::TImplicitObjectInstanced<Chaos::FConvex>>(ConvexImplicit);
+		}
+		else
+		{
+			Implicit = MakeImplicitObjectPtr<Chaos::TImplicitObjectScaled<Chaos::FConvex>>(ConvexImplicit, NetScale);
+		}
+
+		// Wrap in a transform if translation or rotation is non-trivial (scale is pulled out above)
+		if (bHasTranslationOrRotation)
+		{
+			Implicit = MakeImplicitObjectPtr<FImplicitTransformed>(MoveTemp(Implicit), ConvexTransform);
+		}
+
+		return Implicit;
 	}
 
 	Chaos::FImplicitObjectPtr MakeImplicit(const FTriangleMeshGeometry& Geometry, const FRigidShapeInstanceSetup& InSetup)
 	{
-		ensureMsgf(false, TEXT("TODO"));
-		return nullptr;
+		Chaos::FTriangleMeshImplicitObjectPtr TrimeshImplicit = Geometry.GetImplicit();
+		if (!ensure(TrimeshImplicit))
+		{
+			return nullptr;
+		}
+
+		const FTransform3f& InLocalTransform = InSetup.LocalTransform;
+		const FVector InScale(InLocalTransform.GetScale3D());
+		const FQuat InRotation(InLocalTransform.GetRotation());
+		const FVector InTranslation(InLocalTransform.GetTranslation());
+
+		const FTransform MeshTransform = FTransform(InRotation, InTranslation, FVector(1, 1, 1));
+		const bool bHasTranslationOrRotation = !MeshTransform.GetTranslation().IsNearlyZero() || !MeshTransform.GetRotation().IsIdentity();
+
+		// Extract the scale from the transform - we have separate wrapper classes for scale versus translate/rotate 
+		const bool bNoScale = FVector::PointsAreNear(InScale, FVector(1), UE_KINDA_SMALL_NUMBER);
+
+		Chaos::FImplicitObjectPtr Implicit;
+		if (bNoScale)
+		{
+			Implicit = MakeImplicitObjectPtr<Chaos::TImplicitObjectInstanced<Chaos::FTriangleMeshImplicitObject>>(TrimeshImplicit);
+		}
+		else
+		{
+			Implicit = MakeImplicitObjectPtr<Chaos::TImplicitObjectScaled<Chaos::FTriangleMeshImplicitObject>>(TrimeshImplicit, InScale);
+		}
+
+		// Wrap the mesh in a non-scaled transform if necessary (the scale is pulled out above)
+		if (bHasTranslationOrRotation)
+		{
+			Implicit = MakeImplicitObjectPtr<FImplicitObjectTransformed>(MoveTemp(Implicit), MeshTransform);
+		}
+		return Implicit;
+	}
+
+	Chaos::FImplicitObjectPtr MakeImplicit(const FHeightFieldGeometry& Geometry, const FRigidShapeInstanceSetup& InSetup)
+	{
+		const TRefCountPtr<Chaos::FHeightField> HeightFieldImplicit = Geometry.GetImplicit();
+		if (!ensure(HeightFieldImplicit))
+		{
+			return nullptr;
+		}
+
+		const FTransform3f& InLocalTransform = InSetup.LocalTransform;
+		const FVector InScale(InLocalTransform.GetScale3D());
+		const FQuat InRotation(InLocalTransform.GetRotation());
+		const FVector InTranslation(InLocalTransform.GetTranslation());
+
+		const FTransform MeshTransform = FTransform(InRotation, InTranslation, FVector(1, 1, 1));
+		const bool bHasTranslationOrRotation = !MeshTransform.GetTranslation().IsNearlyZero() || !MeshTransform.GetRotation().IsIdentity();
+		const bool bNoScale = FVector::PointsAreNear(InScale, FVector(1), UE_KINDA_SMALL_NUMBER);
+
+		Chaos::FImplicitObjectPtr Implicit;
+		if (bNoScale)
+		{
+			Implicit = MakeImplicitObjectPtr<Chaos::TImplicitObjectInstanced<Chaos::FHeightField>>(HeightFieldImplicit);
+		}
+		else
+		{
+			Implicit = MakeImplicitObjectPtr<Chaos::TImplicitObjectScaled<Chaos::FHeightField>>(HeightFieldImplicit, InScale);
+		}
+
+		if (bHasTranslationOrRotation)
+		{
+			Implicit = MakeImplicitObjectPtr<FImplicitTransformed>(MoveTemp(Implicit), MeshTransform);
+		}
+		return Implicit;
 	}
 
 	template <typename GeometryType>
@@ -149,16 +256,9 @@ namespace Chaos::Rigids::Async
 
 	UE::Physics::FAnyGeometry FRigidShapeInstanceAsync::GetGeometry() const
 	{
-		using FImplicitTransformed = TImplicitObjectTransformed<FReal, 3>;
-		const FImplicitObject* ImplicitGeom = PerShapeData->GetGeometry();
-		EImplicitObjectType ImplicitType = ImplicitGeom->GetType();
-		// TODO: Looks like there's other "wrapper" types like instanced and scaled that should be handled.
-		if (ImplicitType == ImplicitObjectType::Transformed)
-		{
-			const FImplicitTransformed& Transformed = ImplicitGeom->GetObjectChecked<FImplicitTransformed>();
-			ImplicitGeom = Transformed.GetTransformedObject();
-			ImplicitType = ImplicitGeom->GetType();
-		}
+		FImplicitObject* ImplicitGeom = PerShapeData->GetGeometry();
+		auto EmptyCallback = []<typename ImplicitType>(ImplicitType& Implicit){};
+		const EImplicitObjectType ImplicitType = GetLeafImplicit(ImplicitGeom, EmptyCallback);
 
 		switch (ImplicitType)
 		{
@@ -175,9 +275,24 @@ namespace Chaos::Rigids::Async
 				const FImplicitCapsule3& Capsule = ImplicitGeom->GetObjectChecked<FImplicitCapsule3>();
 				return FCapsuleGeometry(Capsule.GetRadiusf(), Capsule.GetHeightf() / 2.0f);
 			}
+			case ImplicitObjectType::Convex:
+			{
+				FConvex& Convex = ImplicitGeom->GetObjectChecked<FConvex>();
+				return UE::Physics::FConvexGeometry(&Convex);
+			}
+			case ImplicitObjectType::TriangleMesh:
+			{
+				FTriangleMeshImplicitObject& TriMesh = ImplicitGeom->GetObjectChecked<FTriangleMeshImplicitObject>();
+				return UE::Physics::FTriangleMeshGeometry(&TriMesh);
+			}
+			case ImplicitObjectType::HeightField:
+			{
+				FHeightField& HeightField = ImplicitGeom->GetObjectChecked<FHeightField>();
+				return UE::Physics::FHeightFieldGeometry(&HeightField);
+			}
 			default:
 			{
-				ensureMsgf(false, TEXT("TODO"));
+				ensureMsgf(false, TEXT("Unhandled implicit leaf type"));
 				break;
 			}
 		}
@@ -186,39 +301,71 @@ namespace Chaos::Rigids::Async
 
 	FTransform3f FRigidShapeInstanceAsync::GetLocalTransform() const
 	{
-		using FImplicitTransformed = TImplicitObjectTransformed<FReal, 3>;
+		FTransform3f Transform;
+		auto AccumulateInternalTransformCallback = [&Transform]<typename ImplicitType>(ImplicitType& Implicit)
+		{
+			if constexpr (std::is_same_v<ImplicitType, FImplicitTransformed>)
+			{
+				FTransform3f ImplicitTransform = FTransform3f(Implicit.GetTransform());
+				FTransform3f::Multiply(&Transform, &ImplicitTransform, &Transform);
+			}
+			else if constexpr (std::is_same_v<ImplicitType, FImplicitObjectScaled>)
+			{
+				FTransform3f ImplicitTransform;
+				ImplicitTransform.SetScale3D(FVector3f(Implicit.GetScale()));
+				FTransform3f::Multiply(&Transform, &ImplicitTransform, &Transform);
+			}
+		};
+
+		// Walk the hierarchy to get a leaf implicit, accumulating the transform as we go down.
+		FImplicitObjectRef ImplicitGeom = PerShapeData->GetGeometry();
+		const EImplicitObjectType ImplicitType = GetLeafImplicit(ImplicitGeom, AccumulateInternalTransformCallback);
 
-		const FImplicitObjectRef Geom = PerShapeData->GetGeometry();
-		switch (Geom->GetType())
+		// Now accumulate the transform of the leaf object.
+		// For most types there's commonly only an internal or leaf transform, however nothing technically prevents both from existing.
+		switch (ImplicitType)
 		{
 			case ImplicitObjectType::Sphere:
 			{
-				FImplicitSphere3& Sphere = Geom->GetObjectChecked<FImplicitSphere3>();
-				return FTransform3f(Sphere.GetCenterf());
+				const FImplicitSphere3& Sphere = ImplicitGeom->GetObjectChecked<FImplicitSphere3>();
+				FTransform3f LeafTransform(Sphere.GetCenterf());
+				FTransform3f::Multiply(&Transform, &LeafTransform, &Transform);
+				break;
 			}
 			case ImplicitObjectType::Box:
 			{
-				FImplicitBox3& Box = Geom->GetObjectChecked<FImplicitBox3>();
-				return FTransform3f(FVector3f(Box.Center()));
+				const FImplicitBox3& Box = ImplicitGeom->GetObjectChecked<FImplicitBox3>();
+				FTransform3f LeafTransform(FVector3f(Box.Center()));
+				FTransform3f::Multiply(&Transform, &LeafTransform, &Transform);
+				break;
 			}
 			case ImplicitObjectType::Capsule:
 			{
-				FImplicitCapsule3& Capsule = Geom->GetObjectChecked<FImplicitCapsule3>();
-				FQuat4f Rotation = FQuat4f::FindBetween(FVector3f(0, 0, 1), Capsule.GetAxisf());
-				return FTransform3f(Rotation, Capsule.GetCenterf());
+				const FImplicitCapsule3& Capsule = ImplicitGeom->GetObjectChecked<FImplicitCapsule3>();
+				const FQuat4f Rotation = FQuat4f::FindBetween(FVector3f(0, 0, 1), Capsule.GetAxisf());
+				FTransform3f LeafTransform(Rotation, Capsule.GetCenterf());
+				FTransform3f::Multiply(&Transform, &LeafTransform, &Transform);
+				break;
+			}
+			case ImplicitObjectType::Convex:
+			{
+				break;
+			}
+			case ImplicitObjectType::TriangleMesh:
+			{
+				break;
 			}
-			case ImplicitObjectType::Transformed:
+			case ImplicitObjectType::HeightField:
 			{
-				const FImplicitTransformed& Transformed = Geom->GetObjectChecked<FImplicitTransformed>();
-				return FTransform3f(Transformed.GetTransform());
+				break;
 			}
 			default:
 			{
-				ensureMsgf(false, TEXT("TODO"));
+				ensureMsgf(false, TEXT("Unhandled implicit leaf type"));
 				break;
 			}
 		}
-		return FTransform3f();
+		return Transform;
 	}
 
 	bool FRigidShapeInstanceAsync::GetQueryEnabled() const
@@ -290,6 +437,41 @@ namespace Chaos::Rigids::Async
 	{
 		PerShapeData->SetUserData(InUserData);
 	}
+
+	template <typename Lambda>
+	EImplicitObjectType FRigidShapeInstanceAsync::GetLeafImplicit(FImplicitObject*& InOutImplicit, const Lambda& Func)
+	{
+		EImplicitObjectType ImplicitType = InOutImplicit->GetType();
+		while (true)
+		{
+			if (ImplicitType == ImplicitObjectType::Transformed)
+			{
+				FImplicitTransformed& Transformed = InOutImplicit->GetObjectChecked<FImplicitTransformed>();
+				Func(Transformed);
+				InOutImplicit = Transformed.MObject.GetReference();
+				ImplicitType = InOutImplicit->GetType();
+				continue;
+			}
+			else if (Chaos::IsInstanced(ImplicitType))
+			{
+				FImplicitObjectInstanced* Instanced = InOutImplicit->AsAChecked<FImplicitObjectInstanced>();
+				Func(*Instanced);
+				InOutImplicit = Instanced->GetInnerObjectInternal();
+				ImplicitType = InOutImplicit->GetType();
+				continue;
+			}
+			else if (Chaos::IsScaled(ImplicitType))
+			{
+				FImplicitObjectScaled* Scaled = InOutImplicit->AsAChecked<FImplicitObjectScaled>();
+				Func(*Scaled);
+				InOutImplicit = Scaled->GetInnerObjectInternal();
+				ImplicitType = InOutImplicit->GetType();
+				continue;
+			}
+			break;
+		}
+		return ImplicitType;
+	}
 } // namespace Chaos::Rigids::Async
 
 #endif // UE_RIGIDPHYSICS_API_ENABLED
diff --git a/Engine/Plugins/Experimental/ChaosRigidPhysicsAsync/Source/ChaosRigidPhysicsAsync/Public/ChaosRigidPhysicsAsync/RigidShapeInstanceAsync.h b/Engine/Plugins/Experimental/ChaosRigidPhysicsAsync/Source/ChaosRigidPhysicsAsync/Public/ChaosRigidPhysicsAsync/RigidShapeInstanceAsync.h
--- a/Engine/Plugins/Experimental/ChaosRigidPhysicsAsync/Source/ChaosRigidPhysicsAsync/Public/ChaosRigidPhysicsAsync/RigidShapeInstanceAsync.h
+++ b/Engine/Plugins/Experimental/ChaosRigidPhysicsAsync/Source/ChaosRigidPhysicsAsync/Public/ChaosRigidPhysicsAsync/RigidShapeInstanceAsync.h
@@ -51,6 +51,10 @@ namespace Chaos::Rigids::Async
 		UE_INTERNAL Chaos::FPerShapeData* PerShapeData = nullptr;
 		UE_INTERNAL UE::Physics::FRigidShapeInstanceHandle Handle;
 		UE_INTERNAL bool bAttachedToBody = false;
+
+	private:
+		template <typename Lambda>
+		static EImplicitObjectType GetLeafImplicit(FImplicitObject*& InOutImplicit, const Lambda& Func);
 	};
 } // namespace Chaos::Rigids::Async
 
diff --git a/Engine/Plugins/Experimental/ChaosRigidPhysicsAsync/Tests/Private/AnyGeometryTests.cpp b/Engine/Plugins/Experimental/ChaosRigidPhysicsAsync/Tests/Private/AnyGeometryTests.cpp
--- a/Engine/Plugins/Experimental/ChaosRigidPhysicsAsync/Tests/Private/AnyGeometryTests.cpp
+++ b/Engine/Plugins/Experimental/ChaosRigidPhysicsAsync/Tests/Private/AnyGeometryTests.cpp
@@ -14,7 +14,9 @@ namespace Chaos::LowLevelTest
 	using FSphereGeometry = UE::Physics::FSphereGeometry;
 	using FBoxGeometry = UE::Physics::FBoxGeometry;
 	using FCapsuleGeometry = UE::Physics::FCapsuleGeometry;
+	using FConvexGeometry = UE::Physics::FConvexGeometry;
 	using FTriangleMeshGeometry = UE::Physics::FTriangleMeshGeometry;
+	using FHeightFieldGeometry = UE::Physics::FHeightFieldGeometry;
 	using FAnyGeometry = UE::Physics::FAnyGeometry;
 
 	template <typename ExpectedGeometryType, typename ActualGeometryType>
@@ -174,6 +176,45 @@ namespace Chaos::LowLevelTest
 		}
 	}
 
+	TEST_CASE("AnyGeometry: Convex", "[Chaos][API][Geometry][unit]")
+	{
+		FBoxGeometry Box0(FVector3f(1, 2, 3));
+		FConvexGeometry Convex0(UE::Physics::MakeConvexBoxGeometrySetup());
+
+		SECTION("IsA")
+		{
+			FAnyGeometry AnyGeometry(Convex0);
+			CHECK(AnyGeometry.IsA<FConvexGeometry>());
+			CHECK(!AnyGeometry.IsA<FBoxGeometry>());
+		}
+		SECTION("TryGet/Get")
+		{
+			FAnyGeometry AnyGeometry(Convex0);
+			REQUIRE(AnyGeometry.TryGet<FConvexGeometry>() != nullptr);
+			CHECK(*AnyGeometry.TryGet<FConvexGeometry>() == Convex0);
+			CHECK(AnyGeometry.Get<FConvexGeometry>() == Convex0);
+			CHECK(AnyGeometry.TryGet<FBoxGeometry>() == nullptr);
+		}
+		SECTION("Visit")
+		{
+			FAnyGeometry AnyGeometry(Convex0);
+			AnyGeometry.Visit([&Convex0]<typename GeometryType>(const GeometryType& Geometry) { CheckGeometry(Convex0, Geometry); });
+		}
+		SECTION("Comparison")
+		{
+			FAnyGeometry AnyGeometry(Convex0);
+			CHECK(AnyGeometry == Convex0);
+			CHECK(AnyGeometry != Box0);
+		}
+		SECTION("Assignment")
+		{
+			FAnyGeometry AnyGeometry(Box0);
+			AnyGeometry = Convex0;
+			CHECK(AnyGeometry != Box0);
+			CHECK(AnyGeometry == Convex0);
+		}
+	}
+
 	TEST_CASE("AnyGeometry: TriangleMesh", "[Chaos][API][Geometry][unit]")
 	{
 		FBoxGeometry Box0(FVector3f(1, 2, 3));
@@ -212,6 +253,44 @@ namespace Chaos::LowLevelTest
 			CHECK(AnyGeometry == TriMesh0);
 		}
 	}
+	TEST_CASE("AnyGeometry: HeightField", "[Chaos][API][Geometry][unit]")
+	{
+		FBoxGeometry Box0(FVector3f(1, 2, 3));
+		FHeightFieldGeometry HeightField0(UE::Physics::MakeHeightFieldGeometrySetup());
+
+		SECTION("IsA")
+		{
+			FAnyGeometry AnyGeometry(HeightField0);
+			CHECK(AnyGeometry.IsA<FHeightFieldGeometry>());
+			CHECK(!AnyGeometry.IsA<FBoxGeometry>());
+		}
+		SECTION("TryGet/Get")
+		{
+			FAnyGeometry AnyGeometry(HeightField0);
+			REQUIRE(AnyGeometry.TryGet<FHeightFieldGeometry>() != nullptr);
+			CHECK(*AnyGeometry.TryGet<FHeightFieldGeometry>() == HeightField0);
+			CHECK(AnyGeometry.Get<FHeightFieldGeometry>() == HeightField0);
+			CHECK(AnyGeometry.TryGet<FBoxGeometry>() == nullptr);
+		}
+		SECTION("Visit")
+		{
+			FAnyGeometry AnyGeometry(HeightField0);
+			AnyGeometry.Visit([&HeightField0]<typename GeometryType>(const GeometryType& Geometry) { CheckGeometry(HeightField0, Geometry); });
+		}
+		SECTION("Comparison")
+		{
+			FAnyGeometry AnyGeometry(HeightField0);
+			CHECK(AnyGeometry == HeightField0);
+			CHECK(AnyGeometry != Box0);
+		}
+		SECTION("Assignment")
+		{
+			FAnyGeometry AnyGeometry(Box0);
+			AnyGeometry = HeightField0;
+			CHECK(AnyGeometry != Box0);
+			CHECK(AnyGeometry == HeightField0);
+		}
+	}
 } // namespace Chaos::LowLevelTest
 
 #endif // UE_RIGIDPHYSICS_API_ENABLED
diff --git a/Engine/Plugins/Experimental/ChaosRigidPhysicsAsync/Tests/Private/RigidTestUtils.cpp b/Engine/Plugins/Experimental/ChaosRigidPhysicsAsync/Tests/Private/RigidTestUtils.cpp
--- a/Engine/Plugins/Experimental/ChaosRigidPhysicsAsync/Tests/Private/RigidTestUtils.cpp
+++ b/Engine/Plugins/Experimental/ChaosRigidPhysicsAsync/Tests/Private/RigidTestUtils.cpp
@@ -68,6 +68,23 @@ namespace UE::Physics
 		return Setup;
 	}
 
+	FHeightFieldGeometrySetup MakeHeightFieldGeometrySetup()
+	{
+		return FHeightFieldGeometrySetup
+		{
+			.Heights
+			{
+				1, 0, 1,
+				0, 2, 0,
+				1, 0, 1
+			},
+			.MaterialIndices = { 0, 1, 2, 3 },
+			.NumRows = 3,
+			.NumCols = 3,
+			.Scale = FVector3f::OneVector,
+		};
+	}
+
 	FRigidShapeInstanceSetup MakeBoxShape(const FVector3f& InSize, const FVector3f& InCenter, UPhysicalMaterial* InMaterial)
 	{
 		FBoxGeometry BoxGeom(InSize);
diff --git a/Engine/Plugins/Experimental/ChaosRigidPhysicsAsync/Tests/Private/RigidTestUtils.h b/Engine/Plugins/Experimental/ChaosRigidPhysicsAsync/Tests/Private/RigidTestUtils.h
--- a/Engine/Plugins/Experimental/ChaosRigidPhysicsAsync/Tests/Private/RigidTestUtils.h
+++ b/Engine/Plugins/Experimental/ChaosRigidPhysicsAsync/Tests/Private/RigidTestUtils.h
@@ -22,6 +22,7 @@ namespace UE::Physics
 	FTriangleMeshGeometrySetup MakeSingleTriangleGeometrySetup(const FVector3f& InScale = FVector3f::One(), const FVector3f& InTranslation = FVector3f::Zero());
 	FTriangleMeshGeometrySetup MakeQuadTriangleGeometrySetup(const FVector3f& InScale = FVector3f::One(), const FVector3f& InTranslation = FVector3f::Zero());
 	FConvexGeometrySetup MakeConvexBoxGeometrySetup(const FVector3f& InCenter = FVector3f::Zero(), const FVector3f& InHalfExtent = FVector3f::One(), const float InMargin = 0);
+	FHeightFieldGeometrySetup MakeHeightFieldGeometrySetup();
 	FRigidShapeInstanceSetup MakeBoxShape(const FVector3f& InSize, const FVector3f& InCenter = FVector3f::Zero(), UPhysicalMaterial* InMaterial = nullptr);
 	FVector3f MakeSolidBoxInertia(float InMass, const FVector3f& InSize);
 } // namespace UE::Physics
diff --git a/Engine/Plugins/Experimental/ChaosRigidPhysicsAsync/Tests/Private/ShapeInstanceTests.cpp b/Engine/Plugins/Experimental/ChaosRigidPhysicsAsync/Tests/Private/ShapeInstanceTests.cpp
--- a/Engine/Plugins/Experimental/ChaosRigidPhysicsAsync/Tests/Private/ShapeInstanceTests.cpp
+++ b/Engine/Plugins/Experimental/ChaosRigidPhysicsAsync/Tests/Private/ShapeInstanceTests.cpp
@@ -1,5 +1,6 @@
 // Copyright Epic Games, Inc. All Rights Reserved.
 
+#include "Chaos/HeightField.h"
 #include "Chaos/LowLevelTest/ChaosTestHarness.h"
 #include "Chaos/LowLevelTest/ChaosTestScene.h"
 #include "ChaosRigidPhysicsAsyncTest.h"
@@ -8,6 +9,7 @@
 #include "RigidPhysics/RigidScene.h"
 #include "RigidPhysics/RigidShapeInstance.h"
 #include "RigidTestFixture.h"
+#include "RigidTestUtils.h"
 #include <catch2/generators/catch_generators_range.hpp>
 
 #if UE_RIGIDPHYSICS_API_ENABLED
@@ -26,14 +28,13 @@ namespace Chaos::LowLevelTest
 	{
 		CHECK_THAT(Chaos::FVec3f(Expected.GetExtents()), Catch::ApproxEq(Chaos::FVec3f(Actual.GetExtents()), UE_KINDA_SMALL_NUMBER));
 	}
-
+	
 	void CheckGeometry(const FSphereGeometry& Expected, const FSphereGeometry& Actual)
 	{
 		CHECK_THAT(Expected.GetRadius(), Catch::Matchers::WithinAbs(Actual.GetRadius(), UE_KINDA_SMALL_NUMBER));
 	}
 
-	template <>
-	void CheckGeometry<FCapsuleGeometry>(const FCapsuleGeometry& Expected, const FCapsuleGeometry& Actual)
+	void CheckGeometry(const FCapsuleGeometry& Expected, const FCapsuleGeometry& Actual)
 	{
 		CHECK_THAT(Expected.GetRadius(), Catch::Matchers::WithinAbs(Actual.GetRadius(), UE_KINDA_SMALL_NUMBER));
 		CHECK_THAT(Expected.GetSegmentHalfLength(), Catch::Matchers::WithinAbs(Actual.GetSegmentHalfLength(), UE_KINDA_SMALL_NUMBER));
@@ -314,6 +315,64 @@ namespace Chaos::LowLevelTest
 				CreateAndCheckShapeInstanceAndTransform(Fixture, Setup, Expected, FTransform3f());
 			}
 		}
+		SECTION("Convex")
+		{
+			FConvexGeometry Geometry(UE::Physics::MakeConvexBoxGeometrySetup());
+			FRigidShapeInstanceSetup Setup(Geometry);
+
+			SECTION("Basic Transforms")
+			{
+				const FTransform3f LocalTransform = GENERATE(
+					FTransform3f(),
+					FTransform3f(FVector3f(1, 2, 3)),
+					FTransform3f(FRotator3f(10, 20, 30)),
+					FTransform3f(FQuat4f::Identity, FVector3f::Zero(), FVector3f(2)),
+					FTransform3f(FQuat4f::Identity, FVector3f::Zero(), FVector3f(2, 3, 4)),
+					FTransform3f(FRotator3f(10, 20, 30), FVector3f(1, 2, 3), FVector3f(2, 3, 4)),
+					// Negative scales
+					FTransform3f(FQuat4f::Identity, FVector3f::Zero(), FVector3f(-2, -3, -4)),
+					FTransform3f(FQuat4f::Identity, FVector3f::Zero(), FVector3f(-2, 1, 1))
+				);
+				Setup.LocalTransform = LocalTransform;
+				CreateAndCheckShapeInstanceAndTransform(Fixture, Setup);
+			}
+			SECTION("Zero Scale")
+			{
+				Setup.LocalTransform.SetScale3D(FVector3f::Zero());
+				const FTransform3f ExpectedTransform(FQuat4f::Identity, FVector3f::Zero(), FVector3f(UE_KINDA_SMALL_NUMBER));
+				CreateAndCheckShapeInstanceAndTransform(Fixture, Setup, Setup.Geometry, ExpectedTransform);
+			}
+		}
+		SECTION("TriangleMesh")
+		{
+			FTriangleMeshGeometry TriMeshGeometry(UE::Physics::MakeQuadTriangleGeometrySetup());
+			FRigidShapeInstanceSetup Setup(TriMeshGeometry);
+			const FTransform3f LocalTransform = GENERATE(
+				FTransform3f(),
+				FTransform3f(FVector3f(0, 1, 0)),
+				FTransform3f(FRotator3f(10, 20, 30)),
+				FTransform3f(FRotator3f(), FVector3f(1, 1, 1), FVector3f(0.5f)),
+				FTransform3f(FRotator3f(), FVector3f(1, 1, 1), FVector3f(2)),
+				FTransform3f(FRotator3f(10, 20, 30), FVector3f(0, 1, 0), FVector3f(2, 3, 4))
+			);
+			Setup.LocalTransform = LocalTransform;
+			CreateAndCheckShapeInstanceAndTransform(Fixture, Setup);
+		}
+		SECTION("HeightField")
+		{
+			FHeightFieldGeometry HeightFieldGeom(UE::Physics::MakeHeightFieldGeometrySetup());
+			FRigidShapeInstanceSetup Setup(HeightFieldGeom);
+			const FTransform3f LocalTransform = GENERATE(
+				FTransform3f(),
+				FTransform3f(FVector3f(0, 1, 0)),
+				FTransform3f(FRotator3f(10, 20, 30)),
+				FTransform3f(FRotator3f(), FVector3f(1, 1, 1), FVector3f(0.5f)),
+				FTransform3f(FRotator3f(), FVector3f(1, 1, 1), FVector3f(2)),
+				FTransform3f(FRotator3f(10, 20, 30), FVector3f(0, 1, 0), FVector3f(2, 3, 4))
+			);
+			Setup.LocalTransform = LocalTransform;
+			CreateAndCheckShapeInstanceAndTransform(Fixture, Setup);
+		}
 	}
 
 	TEST_CASE("ShapeInstance: QueryEnabled", "[Chaos][API][ShapeInstance][unit]")
@@ -583,6 +642,68 @@ namespace Chaos::LowLevelTest
 				CHECK(ShapeInstancePtr->GetUserData() == &UserData[1]);
 			});
 	}
+
+	// This test verifies that creating a shape instance from the geometry of another shape instance shares the implicit.
+	// This is tied to the current implementation of Chaos and should eventually be updated to reflect the new storage when that's completed.
+	TEST_CASE("ShapeInstance: SharedImplicit", "[Chaos][API][ShapeInstance][unit]")
+	{
+		FShapeInstanceTestFixture Fixture;
+		auto TestSharedImplicits = [&Fixture]<typename GeometryType>(const GeometryType& InGeometry)
+		{
+			FRigidShapeInstanceHandle ShapeHandle0, ShapeHandle1;
+			if (FRigidContextGameRW Context = Fixture.LockRWChecked())
+			{
+				TRigidBodyPtr<FRigidContextGameRW> BodyPtr = Context->CreateBody(FRigidDebugName("Body1"), ERigidMovementType::Dynamic);
+				Fixture.AutoCleanup(BodyPtr);
+
+				FRigidShapeInstanceSetup Setup0(InGeometry);
+				TRigidShapeInstancePtr<FRigidContextGameRW> ShapeInstance0 = BodyPtr->CreateShape(Setup0);
+				ShapeHandle0 = ShapeInstance0;
+
+				// Create another shape instance with the same geometry but some different instance data
+				FRigidShapeInstanceSetup Setup1(ShapeInstance0->GetGeometry());
+				Setup1.LocalTransform.SetTranslation(FVector3f(1, 2, 3));
+				ShapeHandle1 = BodyPtr->CreateShape(Setup1);
+				BodyPtr->Activate();
+			}
+
+			auto CheckSharedImplicits = [&Fixture, &ShapeHandle0, &ShapeHandle1]<typename ContextType>(const ContextType& Context)
+			{
+				TRigidShapeInstancePtr<ContextType> ShapeInstance0 = ShapeHandle0.Pin(Context);
+				TRigidShapeInstancePtr<ContextType> ShapeInstance1 = ShapeHandle1.Pin(Context);
+				REQUIRE(ShapeInstance0);
+				REQUIRE(ShapeInstance1);
+
+				const FAnyGeometry AnyGeometry0 = ShapeInstance0->GetGeometry();
+				const FAnyGeometry AnyGeometry1 = ShapeInstance1->GetGeometry();
+				const GeometryType& TypedGeometry0 = AnyGeometry0.Get<GeometryType>();
+				const GeometryType& TypedGeometry1 = AnyGeometry1.Get<GeometryType>();
+
+				CHECK(TypedGeometry0.GetImplicit() == TypedGeometry1.GetImplicit());
+			};
+
+			if (FRigidContextGameRO Context = Fixture.LockROChecked())
+			{
+				CheckSharedImplicits(Context);
+			}
+			Fixture.RunPTCallback(CheckSharedImplicits);
+		};
+		SECTION("Convex")
+		{
+			FConvexGeometry ConvexGeom(UE::Physics::MakeConvexBoxGeometrySetup());
+			TestSharedImplicits(ConvexGeom);
+		}
+		SECTION("TriangleMesh")
+		{
+			FTriangleMeshGeometry TriMeshGeometry(UE::Physics::MakeQuadTriangleGeometrySetup());
+			TestSharedImplicits(TriMeshGeometry);
+		}
+		SECTION("HeightField")
+		{
+			FHeightFieldGeometry HeightFieldGeom(UE::Physics::MakeHeightFieldGeometrySetup());
+			TestSharedImplicits(HeightFieldGeom);
+		}
+	}
 } // namespace Chaos::LowLevelTest
 
 #endif // UE_RIGIDPHYSICS_API_ENABLED
diff --git a/Engine/Source/Runtime/Experimental/Chaos/Public/Chaos/ImplicitObjectScaled.h b/Engine/Source/Runtime/Experimental/Chaos/Public/Chaos/ImplicitObjectScaled.h
--- a/Engine/Source/Runtime/Experimental/Chaos/Public/Chaos/ImplicitObjectScaled.h
+++ b/Engine/Source/Runtime/Experimental/Chaos/Public/Chaos/ImplicitObjectScaled.h
@@ -17,6 +17,11 @@ namespace Chaos
 {
 struct FMTDInfo;
 
+namespace Rigids::Async
+{
+	class FRigidShapeInstanceAsync;
+} // namespace Rigids::Async
+
 class FImplicitObjectInstanced : public FImplicitObject
 {
 public:
@@ -39,6 +44,12 @@ public:
 
 
 protected:
+	friend class Rigids::Async::FRigidShapeInstanceAsync;
+	UE_INTERNAL virtual FImplicitObjectPtr GetInnerObjectInternal() const
+	{
+		return FImplicitObjectPtr();
+	}
+
 	FRealSingle OuterMargin;
 };
 
@@ -392,6 +403,12 @@ protected:
 	{
 		return new TImplicitObjectInstanced<TConcrete>(Obj->MObject);
 	}
+
+private:	
+	UE_INTERNAL virtual FImplicitObjectPtr GetInnerObjectInternal() const override final
+	{
+		return FImplicitObjectPtr(MObject.GetReference());
+	}
 };
 
 class FImplicitObjectScaled : public FImplicitObject
@@ -436,6 +453,12 @@ public:
 	}
 
 protected:
+	friend class Rigids::Async::FRigidShapeInstanceAsync;
+	UE_INTERNAL virtual FImplicitObjectPtr GetInnerObjectInternal() const
+	{
+		return FImplicitObjectPtr();
+	}
+
 	FVec3f MScale;
 	FVec3f MInvScale;
 	FRealSingle OuterMargin;	//Allows us to inflate the instance before the scale is applied. This is useful when sweeps need to apply a non scale on a geometry with uniform thickness
@@ -1149,6 +1172,10 @@ private:
 		return ScaledB;
 	}
 
+	UE_INTERNAL virtual FImplicitObjectPtr GetInnerObjectInternal() const override final
+	{
+		return FImplicitObjectPtr(MObject.GetReference());
+	}
 };
 
 template <typename TConcrete>
diff --git a/Engine/Source/Runtime/Experimental/Chaos/Public/Chaos/ImplicitObjectTransformed.h b/Engine/Source/Runtime/Experimental/Chaos/Public/Chaos/ImplicitObjectTransformed.h
--- a/Engine/Source/Runtime/Experimental/Chaos/Public/Chaos/ImplicitObjectTransformed.h
+++ b/Engine/Source/Runtime/Experimental/Chaos/Public/Chaos/ImplicitObjectTransformed.h
@@ -12,6 +12,11 @@
 namespace Chaos
 {
 
+namespace Rigids::Async
+{
+	class FRigidShapeInstanceAsync;
+} // namespace Rigids::Async
+
 inline void TImplicitObjectTransformSerializeHelper(FChaosArchive& Ar, FImplicitObjectPtr& Obj)
 {
 	Ar << Obj;
@@ -424,6 +429,7 @@ private:
 	TImplicitObjectTransformed() : FImplicitObject(EImplicitObject::HasBoundingBox, ImplicitObjectType::Transformed) {}
 
 	friend FImplicitObject;	//needed for serialization
+	friend Rigids::Async::FRigidShapeInstanceAsync;
 };
 
 namespace Utilities
diff --git a/Engine/Source/Runtime/Experimental/RigidPhysics/Private/RigidPhysics/Geometry/ConvexGeometry.cpp b/Engine/Source/Runtime/Experimental/RigidPhysics/Private/RigidPhysics/Geometry/ConvexGeometry.cpp
--- a/Engine/Source/Runtime/Experimental/RigidPhysics/Private/RigidPhysics/Geometry/ConvexGeometry.cpp
+++ b/Engine/Source/Runtime/Experimental/RigidPhysics/Private/RigidPhysics/Geometry/ConvexGeometry.cpp
@@ -65,6 +65,11 @@ namespace UE::Physics
 	{
 		return TEXT("Convex");
 	}
+
+	TRefCountPtr<Chaos::FConvex> FConvexGeometry::GetImplicit() const
+	{
+		return Implicit;
+	}
 } // namespace UE::Physics
 
 #endif // UE_RIGIDPHYSICS_API_ENABLED
diff --git a/Engine/Source/Runtime/Experimental/RigidPhysics/Private/RigidPhysics/Geometry/HeightFieldGeometry.cpp b/Engine/Source/Runtime/Experimental/RigidPhysics/Private/RigidPhysics/Geometry/HeightFieldGeometry.cpp
--- a/Engine/Source/Runtime/Experimental/RigidPhysics/Private/RigidPhysics/Geometry/HeightFieldGeometry.cpp
+++ b/Engine/Source/Runtime/Experimental/RigidPhysics/Private/RigidPhysics/Geometry/HeightFieldGeometry.cpp
@@ -78,6 +78,11 @@ namespace UE::Physics
 	{
 		return TEXT("HeightField");
 	}
+
+	TRefCountPtr<Chaos::FHeightField> FHeightFieldGeometry::GetImplicit() const
+	{
+		return Implicit;
+	}
 } // namespace UE::Physics
 
 #endif // UE_RIGIDPHYSICS_API_ENABLED
diff --git a/Engine/Source/Runtime/Experimental/RigidPhysics/Public/RigidPhysics/Geometry/AnyGeometry.h b/Engine/Source/Runtime/Experimental/RigidPhysics/Public/RigidPhysics/Geometry/AnyGeometry.h
--- a/Engine/Source/Runtime/Experimental/RigidPhysics/Public/RigidPhysics/Geometry/AnyGeometry.h
+++ b/Engine/Source/Runtime/Experimental/RigidPhysics/Public/RigidPhysics/Geometry/AnyGeometry.h
@@ -8,6 +8,7 @@
 #include "Misc/TVariant.h"
 #include "RigidPhysics/Geometry/CommonCoreGeometry.h"
 #include "RigidPhysics/Geometry/ConvexGeometry.h"
+#include "RigidPhysics/Geometry/HeightFieldGeometry.h"
 #include "RigidPhysics/Geometry/TriangleMeshGeometry.h"
 
 UE_EXPERIMENTAL(5.8, "The new Chaos API is experimental")
@@ -82,7 +83,7 @@ namespace UE::Physics
 
 		RIGIDPHYSICS_API FString ToString() const;
 	private:
-		TVariant<FSphereGeometry, FBoxGeometry, FCapsuleGeometry, FConvexGeometry, FTriangleMeshGeometry> Geometry;
+		TVariant<FSphereGeometry, FBoxGeometry, FCapsuleGeometry, FConvexGeometry, FTriangleMeshGeometry, FHeightFieldGeometry> Geometry;
 	};
 } // namespace UE::Physics
 
diff --git a/Engine/Source/Runtime/Experimental/RigidPhysics/Public/RigidPhysics/Geometry/ConvexGeometry.h b/Engine/Source/Runtime/Experimental/RigidPhysics/Public/RigidPhysics/Geometry/ConvexGeometry.h
--- a/Engine/Source/Runtime/Experimental/RigidPhysics/Public/RigidPhysics/Geometry/ConvexGeometry.h
+++ b/Engine/Source/Runtime/Experimental/RigidPhysics/Public/RigidPhysics/Geometry/ConvexGeometry.h
@@ -50,6 +50,8 @@ namespace UE::Physics
 
 		RIGIDPHYSICS_API FString ToString() const;
 
+		UE_INTERNAL RIGIDPHYSICS_API TRefCountPtr<Chaos::FConvex> GetImplicit() const;
+
 	private:
 		TRefCountPtr<Chaos::FConvex> Implicit;
 	};
diff --git a/Engine/Source/Runtime/Experimental/RigidPhysics/Public/RigidPhysics/Geometry/HeightFieldGeometry.h b/Engine/Source/Runtime/Experimental/RigidPhysics/Public/RigidPhysics/Geometry/HeightFieldGeometry.h
--- a/Engine/Source/Runtime/Experimental/RigidPhysics/Public/RigidPhysics/Geometry/HeightFieldGeometry.h
+++ b/Engine/Source/Runtime/Experimental/RigidPhysics/Public/RigidPhysics/Geometry/HeightFieldGeometry.h
@@ -47,6 +47,8 @@ namespace UE::Physics
 
 		RIGIDPHYSICS_API FString ToString() const;
 
+		UE_INTERNAL RIGIDPHYSICS_API TRefCountPtr<Chaos::FHeightField> GetImplicit() const;
+
 	private:
 		TRefCountPtr<Chaos::FHeightField> Implicit;
 	};
diff --git a/Engine/Source/Runtime/Experimental/RigidPhysics/Public/RigidPhysics/Geometry/TriangleMeshGeometry.h b/Engine/Source/Runtime/Experimental/RigidPhysics/Public/RigidPhysics/Geometry/TriangleMeshGeometry.h
--- a/Engine/Source/Runtime/Experimental/RigidPhysics/Public/RigidPhysics/Geometry/TriangleMeshGeometry.h
+++ b/Engine/Source/Runtime/Experimental/RigidPhysics/Public/RigidPhysics/Geometry/TriangleMeshGeometry.h
@@ -62,7 +62,7 @@ namespace UE::Physics
 
 		RIGIDPHYSICS_API FString ToString() const;
 
-		UE_INTERNAL FTriangleMeshImplicitObjectPtr GetImplicit() const;
+		UE_INTERNAL RIGIDPHYSICS_API FTriangleMeshImplicitObjectPtr GetImplicit() const;
 
 	private:
 		template <typename IndexType>
