Chaos API: Updated the scene api by splitting Tick into Start/End Tick. Also added WaitOnTick.
#Synced-CL 52632702

diff --git a/Engine/Plugins/Experimental/ChaosRigidPhysicsAsync/Source/ChaosRigidPhysicsAsync/Private/ChaosRigidPhysicsAsync/RigidSceneAsync.cpp b/Engine/Plugins/Experimental/ChaosRigidPhysicsAsync/Source/ChaosRigidPhysicsAsync/Private/ChaosRigidPhysicsAsync/RigidSceneAsync.cpp
--- a/Engine/Plugins/Experimental/ChaosRigidPhysicsAsync/Source/ChaosRigidPhysicsAsync/Private/ChaosRigidPhysicsAsync/RigidSceneAsync.cpp
+++ b/Engine/Plugins/Experimental/ChaosRigidPhysicsAsync/Source/ChaosRigidPhysicsAsync/Private/ChaosRigidPhysicsAsync/RigidSceneAsync.cpp
@@ -501,20 +501,46 @@ namespace Chaos::Rigids::Async
 		return false;
 	}
 
-	void FRigidSceneAsyncGT::Tick(float Dt)
+	void FRigidSceneAsyncGT::StartTick(const float InDt)
 	{
 		CheckGameWriteLocked();
 
-		OnPreSimulate(Dt);
+		OnPreSimulate(InDt);
 
 		if (FPBDRigidsSolver* Solver = GetSolver())
 		{
-			Solver->AdvanceAndDispatch_External(Dt);
+			Solver->AdvanceAndDispatch_External(InDt);
+		}
+	}
+
+	void FRigidSceneAsyncGT::WaitOnTick()
+	{
+		CheckGameWriteLocked();
+		
+		if (FPBDRigidsSolver* Solver = GetSolver())
+		{
 			Solver->WaitOnPendingTasks_External();
+		}
+	}
+
+	void FRigidSceneAsyncGT::EndTick()
+	{
+		CheckGameWriteLocked();
+
+		WaitOnTick();
+		
+		// TODO: CopySolverAccelerationStructure
+		if (FPBDRigidsSolver* Solver = GetSolver())
+		{
 			Solver->UpdateGameThreadStructures();
+			// TODO:
+			// - FlipEventManagerBuffer
+			// - SyncEvents_GameThread
+			// - Concrete.SyncQueryMaterials_External();
 		}
 
-		OnPostSimulate(Dt);
+		OnPostSimulate();
+		// TODO: broadcast old OnPhysScenePostTick?
 	}
 
 	void FRigidSceneAsyncGT::OnRigidBodyTransformUpdated(FRigidBodyAsyncGT* Body)
@@ -570,7 +596,7 @@ namespace Chaos::Rigids::Async
 		}
 	}
 
-	void FRigidSceneAsyncGT::OnPostSimulate(FReal Dt)
+	void FRigidSceneAsyncGT::OnPostSimulate()
 	{
 		using namespace UE::Physics;
 
@@ -642,11 +668,6 @@ namespace Chaos::Rigids::Async
 		return false;
 	}
 
-	void FRigidSceneAsyncPT::Tick(float Dt)
-	{
-		UE_RIGIDPHYSICS_CHECKF(false, TEXT("Not supported from physics thread"));
-	}
-
 	UE::Physics::FRigidSceneHandle FRigidSceneAsyncPT::GetSceneHandle() const
 	{
 		CheckSimContext();
diff --git a/Engine/Plugins/Experimental/ChaosRigidPhysicsAsync/Source/ChaosRigidPhysicsAsync/Public/ChaosRigidPhysicsAsync/RigidSceneAsync.h b/Engine/Plugins/Experimental/ChaosRigidPhysicsAsync/Source/ChaosRigidPhysicsAsync/Public/ChaosRigidPhysicsAsync/RigidSceneAsync.h
--- a/Engine/Plugins/Experimental/ChaosRigidPhysicsAsync/Source/ChaosRigidPhysicsAsync/Public/ChaosRigidPhysicsAsync/RigidSceneAsync.h
+++ b/Engine/Plugins/Experimental/ChaosRigidPhysicsAsync/Source/ChaosRigidPhysicsAsync/Public/ChaosRigidPhysicsAsync/RigidSceneAsync.h
@@ -74,7 +74,9 @@ namespace Chaos::Rigids::Async
 		UE_INTERNAL virtual void RegisterModifier(UE::Physics::IRigidSceneModifier* Modifier) override final;
 		UE_INTERNAL virtual void UnregisterModifier(UE::Physics::IRigidSceneModifier* Modifier) override final;
 		UE_INTERNAL virtual bool Query(const UE::Physics::FRigidQueryDesc& QueryDesc, TArray<UE::Physics::FRigidQueryHit>& OutHits) override final;
-		UE_INTERNAL virtual void Tick(float Dt) override final;
+		UE_INTERNAL virtual void StartTick(const float InDt) override final;
+		UE_INTERNAL virtual void WaitOnTick() override final;
+		UE_INTERNAL virtual void EndTick() override final;
 		UE_INTERNAL virtual FRigidBodyAsyncGT* CreateBody(const UE::Physics::FRigidDebugName& InName, UE::Physics::ERigidMovementType InMovementType) override final;
 		UE_INTERNAL virtual void DestroyBody(UE::Physics::IRigidBody* InBody) override final;
 		UE_INTERNAL virtual FRigidGeometryCollectionAsyncGT* CreateGeometryCollection(const UE::Physics::FRigidDebugName& InName, const FSimulationParameters& InParameters) override final;
@@ -118,7 +120,7 @@ namespace Chaos::Rigids::Async
 		void ConstructScene();
 		void DestructScene();
 		void OnPreSimulate(FReal Dt);
-		void OnPostSimulate(FReal Dt);
+		void OnPostSimulate();
 		void UpdateRigidBodyInAccelerationStructure(FRigidBodyAsyncGT* Body);
 		FPBDRigidsSolver* GetSolver() const;
 
@@ -173,7 +175,6 @@ namespace Chaos::Rigids::Async
 		UE_INTERNAL virtual void RegisterModifier(UE::Physics::IRigidSceneModifier* Modifier) override final;
 		UE_INTERNAL virtual void UnregisterModifier(UE::Physics::IRigidSceneModifier* Modifier) override final;
 		UE_INTERNAL virtual bool Query(const UE::Physics::FRigidQueryDesc& QueryDesc, TArray<UE::Physics::FRigidQueryHit>& OutHits) override final;
-		UE_INTERNAL virtual void Tick(float Dt) override final;
 		//
 		// IRIGIDSCENE END
 		//
diff --git a/Engine/Plugins/Experimental/ChaosRigidPhysicsAsync/Tests/Private/BodyTests.cpp b/Engine/Plugins/Experimental/ChaosRigidPhysicsAsync/Tests/Private/BodyTests.cpp
--- a/Engine/Plugins/Experimental/ChaosRigidPhysicsAsync/Tests/Private/BodyTests.cpp
+++ b/Engine/Plugins/Experimental/ChaosRigidPhysicsAsync/Tests/Private/BodyTests.cpp
@@ -308,7 +308,8 @@ namespace Chaos::LowLevelTest
 
 			for (int32 TickIndex = 0; TickIndex < 100; ++TickIndex)
 			{
-				Context->Tick(0.01);
+				Context->StartTick(0.01);
+				Context->EndTick();
 			}
 
 			CHECK_THAT(Body2Ptr->GetTransform().GetLocation(), Catch::ApproxEq(FVec3(0.0, 0.0, 50.0)));
diff --git a/Engine/Plugins/Experimental/ChaosRigidPhysicsAsync/Tests/Private/JointConstraintTests.cpp b/Engine/Plugins/Experimental/ChaosRigidPhysicsAsync/Tests/Private/JointConstraintTests.cpp
--- a/Engine/Plugins/Experimental/ChaosRigidPhysicsAsync/Tests/Private/JointConstraintTests.cpp
+++ b/Engine/Plugins/Experimental/ChaosRigidPhysicsAsync/Tests/Private/JointConstraintTests.cpp
@@ -186,7 +186,8 @@ namespace Chaos::LowLevelTest
 			TRigidBodyPtr<FRigidContextGameRW> Body1Ptr = Body1Handle.Pin(Context);
 			for (int32 TickIndex = 0; TickIndex < 10; ++TickIndex)
 			{
-				Context->Tick(0.01);
+				Context->StartTick(0.01);
+				Context->EndTick();
 			}
 
 			CHECK(Body1Ptr->GetTransform().GetLocation().Z > -151);
diff --git a/Engine/Plugins/Experimental/ChaosRigidPhysicsAsync/Tests/Private/RigidTestFixture.cpp b/Engine/Plugins/Experimental/ChaosRigidPhysicsAsync/Tests/Private/RigidTestFixture.cpp
--- a/Engine/Plugins/Experimental/ChaosRigidPhysicsAsync/Tests/Private/RigidTestFixture.cpp
+++ b/Engine/Plugins/Experimental/ChaosRigidPhysicsAsync/Tests/Private/RigidTestFixture.cpp
@@ -110,12 +110,13 @@ namespace Chaos::LowLevelTest
 		FTestRigidSceneModifier SceneModifier;
 		SceneModifier.PreTickFunc = Callback;
 
-		if (FRigidContextGameRW Context = SceneHandle.LockRW())
+		if (FRigidContextGameRW Context = LockRWChecked())
 		{
 			Context->RegisterModifier(&SceneModifier);
 			for (int32 I = 0; I < Iterations; ++I)
 			{
-				Context->Tick(Dt);
+				Context->StartTick(Dt);
+				Context->EndTick();
 			}
 			Context->UnregisterModifier(&SceneModifier);
 		}
diff --git a/Engine/Plugins/Experimental/ChaosRigidPhysicsAsync/Tests/Private/SceneTests.cpp b/Engine/Plugins/Experimental/ChaosRigidPhysicsAsync/Tests/Private/SceneTests.cpp
--- a/Engine/Plugins/Experimental/ChaosRigidPhysicsAsync/Tests/Private/SceneTests.cpp
+++ b/Engine/Plugins/Experimental/ChaosRigidPhysicsAsync/Tests/Private/SceneTests.cpp
@@ -159,7 +159,8 @@ namespace Chaos::LowLevelTest
 			{
 				if (FRigidContextGameRW Context = SceneHandle.LockRW())
 				{
-					Context.GetScene()->Tick(Dt);
+					Context.GetScene()->StartTick(Dt);
+					Context.GetScene()->EndTick();
 
 					// Dynamic body should have moved under gravity
 					if (TRigidBodyPtr<FRigidContextGameRW> BodyPtr = DynamicBodyHandle.Pin(Context))
@@ -214,7 +215,8 @@ namespace Chaos::LowLevelTest
 				// Add Collision
 				BodyHandle.Pin(Context)->CreateShape(MakeBoxShape(FVector3f(100, 100, 100)));
 
-				Context->Tick(Dt);
+				Context->StartTick(Dt);
+				Context->EndTick();
 
 				// Destroy the body
 				Context->DestroyBody(BodyHandle.Pin(Context));
@@ -271,17 +273,20 @@ namespace Chaos::LowLevelTest
 		{
 			Context->RegisterModifier(&SceneModifier);
 
-			Context->Tick(Dt);
+			Context->StartTick(Dt);
+			Context->EndTick();
 			CHECK(PreSimCallCount == 1);
 			CHECK(PreTickCallCount == 1);
 
-			Context->Tick(Dt);
+			Context->StartTick(Dt);
+			Context->EndTick();
 			CHECK(PreSimCallCount == 2);
 			CHECK(PreTickCallCount == 2);
 
 			Context->UnregisterModifier(&SceneModifier);
 
-			Context->Tick(Dt);
+			Context->StartTick(Dt);
+			Context->EndTick();
 			CHECK(PreSimCallCount == 2);
 			CHECK(PreTickCallCount == 2);
 		}
@@ -354,7 +359,8 @@ namespace Chaos::LowLevelTest
 			CHECK(BodyPtr1->GetTransform().GetLocation().X == 0.0);
 			CHECK(BodyPtr2->GetTransform().GetLocation().X == 0.0);
 
-			Context->Tick(Dt);
+			Context->StartTick(Dt);
+			Context->EndTick();
 
 			// Both bodies were moved in callbacks
 			CHECK(BodyPtr1->GetTransform().GetLocation().X == 200.0);
diff --git a/Engine/Source/Runtime/Experimental/Dataflow/Simulation/Private/Dataflow/DataflowSimulationNodes.cpp b/Engine/Source/Runtime/Experimental/Dataflow/Simulation/Private/Dataflow/DataflowSimulationNodes.cpp
--- a/Engine/Source/Runtime/Experimental/Dataflow/Simulation/Private/Dataflow/DataflowSimulationNodes.cpp
+++ b/Engine/Source/Runtime/Experimental/Dataflow/Simulation/Private/Dataflow/DataflowSimulationNodes.cpp
@@ -122,7 +122,8 @@ void FAdvanceRigidPhysicsSolverDataflowNode::EvaluateSimulation(UE::Dataflow::FD
 	{
 		if (UE::Physics::FRigidContextGameRW Context = RigidState->Handle.LockRW())
 		{
-			Context->Tick(SimulationDeltaTime);
+			Context->StartTick(SimulationDeltaTime);
+			Context->EndTick();
 		}
 	}
 
diff --git a/Engine/Source/Runtime/Experimental/RigidPhysics/Private/RigidPhysics/Internal/IRigidScene.cpp b/Engine/Source/Runtime/Experimental/RigidPhysics/Private/RigidPhysics/Internal/IRigidScene.cpp
--- a/Engine/Source/Runtime/Experimental/RigidPhysics/Private/RigidPhysics/Internal/IRigidScene.cpp
+++ b/Engine/Source/Runtime/Experimental/RigidPhysics/Private/RigidPhysics/Internal/IRigidScene.cpp
@@ -17,7 +17,17 @@ namespace UE::Physics
 		return false;
 	}
 
-	void IRigidScene::Tick(float InDt)
+	void IRigidScene::StartTick(float InDt)
+	{
+		RIGIDPHYSICS_API_UNSUPPORTED();
+	}
+
+	void IRigidScene::WaitOnTick()
+	{
+		RIGIDPHYSICS_API_UNSUPPORTED();
+	}
+
+	void IRigidScene::EndTick()
 	{
 		RIGIDPHYSICS_API_UNSUPPORTED();
 	}
diff --git a/Engine/Source/Runtime/Experimental/RigidPhysics/Public/RigidPhysics/Internal/IRigidScene.h b/Engine/Source/Runtime/Experimental/RigidPhysics/Public/RigidPhysics/Internal/IRigidScene.h
--- a/Engine/Source/Runtime/Experimental/RigidPhysics/Public/RigidPhysics/Internal/IRigidScene.h
+++ b/Engine/Source/Runtime/Experimental/RigidPhysics/Public/RigidPhysics/Internal/IRigidScene.h
@@ -35,7 +35,9 @@ namespace UE::Physics
 		UE_INTERNAL RIGIDPHYSICS_API virtual bool Query(const FRigidQueryDesc& InQueryDesc, TArray<FRigidQueryHit>& OutHits);
 
 		// Advance the simulation by Dt. If the implementation is asynchronous, this will return immediately
-		UE_INTERNAL RIGIDPHYSICS_API virtual void Tick(float InDt);
+		UE_INTERNAL RIGIDPHYSICS_API virtual void StartTick(const float InDt);
+		UE_INTERNAL RIGIDPHYSICS_API virtual void WaitOnTick();
+		UE_INTERNAL RIGIDPHYSICS_API virtual void EndTick();
 
 		// GT Only API - Move this somehow
 		UE_INTERNAL RIGIDPHYSICS_API virtual void RegisterModifier(IRigidSceneModifier* InModifier);
diff --git a/Engine/Source/Runtime/Experimental/RigidPhysics/Public/RigidPhysics/RigidScene.h b/Engine/Source/Runtime/Experimental/RigidPhysics/Public/RigidPhysics/RigidScene.h
--- a/Engine/Source/Runtime/Experimental/RigidPhysics/Public/RigidPhysics/RigidScene.h
+++ b/Engine/Source/Runtime/Experimental/RigidPhysics/Public/RigidPhysics/RigidScene.h
@@ -118,9 +118,19 @@ namespace UE::Physics
 			return SceneRaw->Query(QueryDesc, OutHits);
 		}
 
-		void Tick(float Dt) requires CIsGameContext<FContext>
+		void StartTick(const float InDt) requires CIsGameContext<FContext>
 		{
-			SceneRaw->Tick(Dt);
+			SceneRaw->StartTick(InDt);
+		}
+
+		void WaitOnTick() requires CIsGameContext<FContext>
+		{
+			SceneRaw->WaitOnTick();
+		}
+
+		void EndTick() requires CIsGameContext<FContext>
+		{
+			SceneRaw->EndTick();
 		}
 
 		UE_INTERNAL void Reset()
